Fix bad merge from physic copy up

diff --git a/Engine/Source/Runtime/ClothingSystemRuntimeCommon/Private/ClothingAsset.cpp b/Engine/Source/Runtime/ClothingSystemRuntimeCommon/Private/ClothingAsset.cpp
--- a/Engine/Source/Runtime/ClothingSystemRuntimeCommon/Private/ClothingAsset.cpp
+++ b/Engine/Source/Runtime/ClothingSystemRuntimeCommon/Private/ClothingAsset.cpp
@@ -41,12 +41,32 @@ void ClothingAssetUtils::GetMeshClothingAssetBindings(
 	{
 		return;
 	}
+#if WITH_EDITORONLY_DATA
+	if (InSkelMesh->GetImportedModel())
+	{
+		int32 LODNum = InSkelMesh->GetImportedModel()->LODModels.Num();
+		for (int32 LODIndex = 0; LODIndex < LODNum; ++LODIndex)
+		{
+			if (InSkelMesh->GetImportedModel()->LODModels[LODIndex].HasClothData())
+			{
+				TArray<FClothingAssetMeshBinding> LodBindings;
+				GetMeshClothingAssetBindings(InSkelMesh, LodBindings, LODIndex);
+				OutBindings.Append(LodBindings);
+			}
+		}
+		if (OutBindings.Num() > 0)
+		{
+			return;
+		}
+	}
+#endif
 
-	if(FSkeletalMeshRenderData* Resource = InSkelMesh->GetResourceForRendering())
+	//Fallback on render data
+	if (FSkeletalMeshRenderData* Resource = InSkelMesh->GetResourceForRendering())
 	{
 		const int32 NumLods = Resource->LODRenderData.Num();
 
-		for(int32 LodIndex = 0; LodIndex < NumLods; ++LodIndex)
+		for (int32 LodIndex = 0; LodIndex < NumLods; ++LodIndex)
 		{
 			TArray<FClothingAssetMeshBinding> LodBindings;
 			GetMeshClothingAssetBindings(InSkelMesh, LodBindings, LodIndex);
@@ -68,6 +88,39 @@ void ClothingAssetUtils::GetMeshClothingAssetBindings(
 		return;
 	}
 
+#if WITH_EDITORONLY_DATA
+	if (InSkelMesh->GetImportedModel())
+	{
+		int32 LODNum = InSkelMesh->GetImportedModel()->LODModels.Num();
+		if (InSkelMesh->GetImportedModel()->LODModels[InLodIndex].HasClothData())
+		{
+			TArray<FClothingAssetMeshBinding> LodBindings;
+			int32 SectionNum = InSkelMesh->GetImportedModel()->LODModels[InLodIndex].Sections.Num();
+			for (int32 SectionIndex = 0; SectionIndex < SectionNum; ++SectionIndex)
+			{
+				const FSkelMeshSection& Section = InSkelMesh->GetImportedModel()->LODModels[InLodIndex].Sections[SectionIndex];
+				if (Section.HasClothingData())
+				{
+					UClothingAssetBase* ClothingAsset = InSkelMesh->GetClothingAsset(Section.ClothingData.AssetGuid);
+					FClothingAssetMeshBinding ClothBinding;
+					ClothBinding.Asset = Cast<UClothingAssetCommon>(ClothingAsset);
+					ClothBinding.AssetInternalLodIndex = Section.ClothingData.AssetLodIndex;// InSkelMesh->GetClothingAssetIndex(Section.ClothingData.AssetGuid);
+					check(ClothBinding.AssetInternalLodIndex == Section.ClothingData.AssetLodIndex);
+					ClothBinding.LODIndex = InLodIndex;
+					ClothBinding.SectionIndex = SectionIndex;
+					OutBindings.Add(ClothBinding);
+				}
+			}
+		}
+
+		if (OutBindings.Num() > 0)
+		{
+			return;
+		}
+	}
+#endif
+
+	//Fallback on render data
 	if(FSkeletalMeshRenderData* Resource = InSkelMesh->GetResourceForRendering())
 	{
 		if(Resource->LODRenderData.IsValidIndex(InLodIndex))
@@ -139,8 +192,7 @@ bool UClothingAssetCommon::BindToSkeletalMesh(
 	USkeletalMesh* InSkelMesh, 
 	const int32 InMeshLodIndex, 
 	const int32 InSectionIndex, 
-	const int32 InAssetLodIndex, 
-	const bool bCallPostEditChange)
+	const int32 InAssetLodIndex)
 {
 	// If we've been added to the wrong mesh
 	if(InSkelMesh != GetOuter())
@@ -300,19 +352,8 @@ bool UClothingAssetCommon::BindToSkeletalMesh(
 	// After verifying copy the new bone map to the section
 	OriginalSection.BoneMap = TempBoneMap;
 
-	// Array of re-import contexts for components using this mesh
-	TIndirectArray<FComponentReregisterContext> ComponentContexts;
-	for (TObjectIterator<USkeletalMeshComponent> It; It; ++It)
-	{
-		USkeletalMeshComponent* Component = *It;
-		if (Component && !Component->IsTemplate() && Component->SkeletalMesh == InSkelMesh)
-		{
-			ComponentContexts.Add(new FComponentReregisterContext(Component));
-		}
-	}
-
-	// Ready to apply the changes
-	InSkelMesh->PreEditChange(nullptr);
+	//Register the scope post edit change
+	FScopedSkeletalMeshPostEditChange SkeletalMeshPostEditChange(InSkelMesh);
 
 	// calculate LOD verts before adding our new section
 	uint32 NumLodVertices = 0;
@@ -364,13 +405,9 @@ bool UClothingAssetCommon::BindToSkeletalMesh(
 
 	LodMap[InMeshLodIndex] = InAssetLodIndex;
 
-	if (bCallPostEditChange)
-	{
-		InSkelMesh->PostEditChange();
-	}
-
 	return true;
-	// ComponentContexts goes out of scope, causing components to be re-registered
+
+	// FScopedSkeletalMeshPostEditChange goes out of scope, causing postedit change and components to be re-registered
 }
 
 void UClothingAssetCommon::UnbindFromSkeletalMesh(USkeletalMesh* InSkelMesh)
@@ -413,10 +450,7 @@ void UClothingAssetCommon::UnbindFromSkeletalMesh(
 			FSkelMeshSection& Section = LodModel.Sections[SectionIdx];
 			if(Section.HasClothingData() && Section.ClothingData.AssetGuid == AssetGuid)
 			{
-				if(!bChangedMesh)
-				{
-					InSkelMesh->PreEditChange(nullptr);
-				}
+				InSkelMesh->PreEditChange(nullptr);
 				ClothingAssetUtils::ClearSectionClothingData(Section);
 				bChangedMesh = true;
 			}
@@ -433,18 +467,8 @@ void UClothingAssetCommon::UnbindFromSkeletalMesh(
 	// If the mesh changed we need to re-register any components that use it to reflect the changes
 	if(bChangedMesh)
 	{
-		InSkelMesh->PostEditChange();
-
-		for(TObjectIterator<USkeletalMeshComponent> It; It; ++It)
-		{
-			USkeletalMeshComponent* MeshComponent = *It;
-			if(MeshComponent &&
-			   !MeshComponent->IsTemplate() &&
-			   MeshComponent->SkeletalMesh == InSkelMesh)
-			{
-				MeshComponent->ReregisterComponent();
-			}
-		}
+		//Register the scope post edit change
+		FScopedSkeletalMeshPostEditChange SkeletalMeshPostEditChange(InSkelMesh);
 	}
 }
 
diff --git a/Engine/Source/Runtime/ClothingSystemRuntimeCommon/Public/ClothingAsset.h b/Engine/Source/Runtime/ClothingSystemRuntimeCommon/Public/ClothingAsset.h
--- a/Engine/Source/Runtime/ClothingSystemRuntimeCommon/Public/ClothingAsset.h
+++ b/Engine/Source/Runtime/ClothingSystemRuntimeCommon/Public/ClothingAsset.h
@@ -98,7 +98,7 @@ public:
 	 * Create weights for skinning the render mesh to our simulation mesh, and 
 	 * weights to drive our sim mesh from the skeleton.
 	 */
-	virtual bool BindToSkeletalMesh(USkeletalMesh* InSkelMesh, const int32 InMeshLodIndex, const int32 InSectionIndex, const int32 InAssetLodIndex, const bool bCallPostEditChange=true) override;
+	virtual bool BindToSkeletalMesh(USkeletalMesh* InSkelMesh, const int32 InMeshLodIndex, const int32 InSectionIndex, const int32 InAssetLodIndex) override;
 
 	/**
 	 * Helper that invokes \c UnbindFromSkeletalMesh() for each avilable entry in 
diff --git a/Engine/Source/Runtime/ClothingSystemRuntimeInterface/Public/ClothingAssetBase.h b/Engine/Source/Runtime/ClothingSystemRuntimeInterface/Public/ClothingAssetBase.h
--- a/Engine/Source/Runtime/ClothingSystemRuntimeInterface/Public/ClothingAssetBase.h
+++ b/Engine/Source/Runtime/ClothingSystemRuntimeInterface/Public/ClothingAssetBase.h
@@ -23,7 +23,7 @@ public:
 	* @param InSubmeshIdx Submesh in this asset to replace section with
 	* @param InAssetLodIndex Internal clothing LOD to use
 	*/
-	virtual bool BindToSkeletalMesh(USkeletalMesh* InSkelMesh, const int32 InMeshLodIndex, const int32 InSectionIndex, const int32 InAssetLodIndex, const bool bCallPostEditChange = true)
+	virtual bool BindToSkeletalMesh(USkeletalMesh* InSkelMesh, const int32 InMeshLodIndex, const int32 InSectionIndex, const int32 InAssetLodIndex)
 	PURE_VIRTUAL(UClothingAssetBase::BindToSkeletalMesh, return false;);
 
 	/**
