Fixed localization conflicts

#rnx

diff --git a/Engine/Plugins/Developer/PlasticSourceControl/Source/PlasticSourceControl/Private/SPlasticSourceControlSettings.cpp b/Engine/Plugins/Developer/PlasticSourceControl/Source/PlasticSourceControl/Private/SPlasticSourceControlSettings.cpp
--- a/Engine/Plugins/Developer/PlasticSourceControl/Source/PlasticSourceControl/Private/SPlasticSourceControlSettings.cpp
+++ b/Engine/Plugins/Developer/PlasticSourceControl/Source/PlasticSourceControl/Private/SPlasticSourceControlSettings.cpp
@@ -61,7 +61,7 @@ void SPlasticSourceControlSettings::Construct(const FArguments& InArgs)
 				.FillWidth(1.0f)
 				[
 					SNew(STextBlock)
-					.Text(LOCTEXT("BinaryPathLabel", "Plastic SCM Path"))
+					.Text(LOCTEXT("PathLabel", "Plastic SCM Path"))
 					.Font(Font)
 				]
 				+SHorizontalBox::Slot()
@@ -206,7 +206,7 @@ void SPlasticSourceControlSettings::Construct(const FArguments& InArgs)
 				[
 					SNew(SEditableTextBox)
 					.Text(this, &SPlasticSourceControlSettings::GetServerUrl)
-					.HintText(LOCTEXT("ServerUrl", "Enter the Server URL"))
+					.HintText(LOCTEXT("EnterServerUrl", "Enter the Server URL"))
 					.OnTextCommitted(this, &SPlasticSourceControlSettings::OnServerUrlCommited)
 					.Font(Font)
 				]
diff --git a/Engine/Plugins/Enterprise/AxFImporter/Source/AxFImporter/Private/AxFImporterFactory.cpp b/Engine/Plugins/Enterprise/AxFImporter/Source/AxFImporter/Private/AxFImporterFactory.cpp
--- a/Engine/Plugins/Enterprise/AxFImporter/Source/AxFImporter/Private/AxFImporterFactory.cpp
+++ b/Engine/Plugins/Enterprise/AxFImporter/Source/AxFImporter/Private/AxFImporterFactory.cpp
@@ -225,11 +225,11 @@ void UAxFImporterFactory::CleanUp()
 bool UAxFImporterFactory::LoadMaterials(IAxFFileImporter* Importer,
     UObject* InPackage, EObjectFlags InFlags, const FString& InFilename, const UAxFImporterOptions& InImporterOptions, FFeedbackContext& OutWarn)
 {
-	FScopedSlowTask Progress(100.0f, LOCTEXT("AxF", "Importing File ..."), true, OutWarn);
+	FScopedSlowTask Progress(100.0f, LOCTEXT("ImportingFile", "Importing File ..."), true, OutWarn);
 	Progress.MakeDialog();
 
 	const float StartProgress = 25.f;
-	Progress.EnterProgressFrame(StartProgress, FText::Format(LOCTEXT("AxF", "Loading materials: {0} ..."), FText::FromString(InFilename)));
+	Progress.EnterProgressFrame(StartProgress, FText::Format(LOCTEXT("LoadingMaterials", "Loading materials: {0} ..."), FText::FromString(InFilename)));
 
 	bool bSuccess = Importer->OpenFile(InFilename, InImporterOptions);
 
@@ -245,11 +245,11 @@ bool UAxFImporterFactory::LoadMaterials(IAxFFileImporter* Importer,
 				const float FinishProgress   = 90.f;
 				const float MaterialProgress = (FinishProgress - StartProgress) / MaterialCount;
 				Progress.EnterProgressFrame(MaterialProgress,
-				                            FText::Format(LOCTEXT("AxF", "Material processing: {0} ..."), FText::FromString(MaterialName)));
+				                            FText::Format(LOCTEXT("MaterialProcessing", "Material processing: {0} ..."), FText::FromString(MaterialName)));
 			}
 			else
 			{
-				Progress.EnterProgressFrame(5.f, FText::Format(LOCTEXT("AxF", "Finishing materials: {0} ..."), FText::FromString(InFilename)));
+				Progress.EnterProgressFrame(5.f, FText::Format(LOCTEXT("FinishingMaterials", "Finishing materials: {0} ..."), FText::FromString(InFilename)));
 			}
 		};
 
diff --git a/Engine/Plugins/Enterprise/DatasmithC4DImporter/Source/DatasmithC4DTranslator/Private/DatasmithC4DTranslatorModule.cpp b/Engine/Plugins/Enterprise/DatasmithC4DImporter/Source/DatasmithC4DTranslator/Private/DatasmithC4DTranslatorModule.cpp
--- a/Engine/Plugins/Enterprise/DatasmithC4DImporter/Source/DatasmithC4DTranslator/Private/DatasmithC4DTranslatorModule.cpp
+++ b/Engine/Plugins/Enterprise/DatasmithC4DImporter/Source/DatasmithC4DTranslator/Private/DatasmithC4DTranslatorModule.cpp
@@ -88,4 +88,6 @@ public:
 };
 #endif //_MELANGE_SDK_
 
-IMPLEMENT_MODULE(FC4DTranslatorModule, DatasmithC4DTranslator);
\ No newline at end of file
+IMPLEMENT_MODULE(FC4DTranslatorModule, DatasmithC4DTranslator);
+
+#undef LOCTEXT_NAMESPACE
diff --git a/Engine/Plugins/Enterprise/DatasmithImporter/Source/DatasmithImporter/Private/DatasmithFileProducer.cpp b/Engine/Plugins/Enterprise/DatasmithImporter/Source/DatasmithImporter/Private/DatasmithFileProducer.cpp
--- a/Engine/Plugins/Enterprise/DatasmithImporter/Source/DatasmithImporter/Private/DatasmithFileProducer.cpp
+++ b/Engine/Plugins/Enterprise/DatasmithImporter/Source/DatasmithImporter/Private/DatasmithFileProducer.cpp
@@ -153,7 +153,7 @@ bool UDatasmithFileProducer::Execute(TArray< TWeakObjectPtr< UObject > >& OutAss
 	ProgressTaskPtr->ReportNextStep( FText::Format( LOCTEXT( "DatasmithFileProducer_ConvertingFile", "Converting {0} ..."), FText::FromString( FilePath ) ), 2.0f );
 	SceneElementToWorld();
 
-	ProgressTaskPtr->ReportNextStep( LOCTEXT( "DatasmithFileProducer_LoadingFile", "Cleaning data ...") );
+	ProgressTaskPtr->ReportNextStep( LOCTEXT( "DatasmithFileProducer_CleaningData", "Cleaning data ...") );
 	PreventNameCollision();
 
 	OutAssets.Append( MoveTemp( Assets ) );
@@ -600,7 +600,7 @@ bool UDatasmithDirProducer::Execute(TArray< TWeakObjectPtr< UObject > >& OutAsse
 {
 	if(!IsValid())
 	{
-		LogError( LOCTEXT( "DatasmithDirProducerNoFile", "Execution failed: Producer is not valid." ) );
+		LogError( LOCTEXT( "DatasmithProducerInvalid", "Execution failed: Producer is not valid." ) );
 		return false;
 	}
 
@@ -867,7 +867,7 @@ public:
 			[
 				SNew(SButton)
 				.OnClicked(this, &SDatasmithFileProducerFileProperty::OnChangePathClicked )
-				.ToolTipText(LOCTEXT("ChangePath_Tooltip", "Browse for a new source file path"))
+				.ToolTipText(LOCTEXT("ChangeSourcePath_Tooltip", "Browse for a new source file path"))
 				[
 					SNew(STextBlock)
 					.Text(LOCTEXT("...", "..."))
diff --git a/Engine/Plugins/Enterprise/DatasmithImporter/Source/DatasmithImporter/Private/DatasmithImporter.cpp b/Engine/Plugins/Enterprise/DatasmithImporter/Source/DatasmithImporter/Private/DatasmithImporter.cpp
--- a/Engine/Plugins/Enterprise/DatasmithImporter/Source/DatasmithImporter/Private/DatasmithImporter.cpp
+++ b/Engine/Plugins/Enterprise/DatasmithImporter/Source/DatasmithImporter/Private/DatasmithImporter.cpp
@@ -1080,7 +1080,7 @@ void FDatasmithImporter::ImportTextures( FDatasmithImportContext& ImportContext
 
 	const int32 TexturesCount = ImportContext.FilteredScene->GetTexturesCount();
 
-	FScopedSlowTask Progress( (float)TexturesCount, LOCTEXT("ImportSlowTask", "Importing Textures..."), true, *ImportContext.Warn );
+	FScopedSlowTask Progress( (float)TexturesCount, LOCTEXT("ImportingTextures", "Importing Textures..."), true, *ImportContext.Warn );
 	Progress.MakeDialog(true);
 
 	if (ImportContext.Options->TextureConflictPolicy != EDatasmithImportAssetConflictPolicy::Ignore && TexturesCount > 0)
@@ -1807,7 +1807,7 @@ void FDatasmithImporter::ImportLevelSequences( FDatasmithImportContext& ImportCo
 		return;
 	}
 
-	FScopedSlowTask Progress( (float)SequencesCount, LOCTEXT("ImportSlowTask", "Importing Level Sequences..."), true, *ImportContext.Warn );
+	FScopedSlowTask Progress( (float)SequencesCount, LOCTEXT("ImportingLevelSequences", "Importing Level Sequences..."), true, *ImportContext.Warn );
 	Progress.MakeDialog(true);
 
 	// We can only parse a IDatasmithLevelSequenceElement with IDatasmithSubsequenceAnimationElements if their target
@@ -1908,7 +1908,7 @@ void FDatasmithImporter::ImportLevelVariantSets( FDatasmithImportContext& Import
 		return;
 	}
 
-	FScopedSlowTask Progress( (float)LevelVariantSetsCount, LOCTEXT("ImportSlowTask", "Importing Level Variant Sets..."), true, *ImportContext.Warn );
+	FScopedSlowTask Progress( (float)LevelVariantSetsCount, LOCTEXT("ImportingLevelVariantSets", "Importing Level Variant Sets..."), true, *ImportContext.Warn );
 	Progress.MakeDialog(true);
 
 	for ( int32 LevelVariantSetIndex = 0; LevelVariantSetIndex < LevelVariantSetsCount && !ImportContext.bUserCancelled; ++LevelVariantSetIndex )
diff --git a/Engine/Plugins/Enterprise/DatasmithImporter/Source/DatasmithImporter/Private/DatasmithStaticMeshImporter.cpp b/Engine/Plugins/Enterprise/DatasmithImporter/Source/DatasmithImporter/Private/DatasmithStaticMeshImporter.cpp
--- a/Engine/Plugins/Enterprise/DatasmithImporter/Source/DatasmithImporter/Private/DatasmithStaticMeshImporter.cpp
+++ b/Engine/Plugins/Enterprise/DatasmithImporter/Source/DatasmithImporter/Private/DatasmithStaticMeshImporter.cpp
@@ -29,7 +29,7 @@
 #include "Templates/UniquePtr.h"
 #include "UObject/Package.h"
 
-#define LOCTEXT_NAMESPACE "DatasmithImporter"
+#define LOCTEXT_NAMESPACE "DatasmithStaticMeshImporter"
 
 namespace DatasmithStaticMeshImporterImpl
 {
diff --git a/Engine/Plugins/Enterprise/DatasmithImporter/Source/DatasmithImporter/Private/DatasmithTextureImporter.cpp b/Engine/Plugins/Enterprise/DatasmithImporter/Source/DatasmithImporter/Private/DatasmithTextureImporter.cpp
--- a/Engine/Plugins/Enterprise/DatasmithImporter/Source/DatasmithImporter/Private/DatasmithTextureImporter.cpp
+++ b/Engine/Plugins/Enterprise/DatasmithImporter/Source/DatasmithImporter/Private/DatasmithTextureImporter.cpp
@@ -191,7 +191,7 @@ UTexture* FDatasmithTextureImporter::CreateTexture(const TSharedPtr<IDatasmithTe
 		// try opening from absolute path
 		if (!(FFileHelper::LoadFileToArray(FileTextureData, *ImageFileName) && FileTextureData.Num() > 0))
 		{
-			ImportContext.LogWarning(FText::Format(LOCTEXT("InvalidData", "Unable to find Texture file {0}."), FText::FromString(Filename)));
+			ImportContext.LogWarning(FText::Format(LOCTEXT("NoTextureFileFound", "Unable to find Texture file {0}."), FText::FromString(Filename)));
 			return nullptr;
 		}
 
diff --git a/Engine/Plugins/Enterprise/DatasmithImporter/Source/DatasmithImporter/Private/UI/DatasmithImportOptionsWindow.cpp b/Engine/Plugins/Enterprise/DatasmithImporter/Source/DatasmithImporter/Private/UI/DatasmithImportOptionsWindow.cpp
--- a/Engine/Plugins/Enterprise/DatasmithImporter/Source/DatasmithImporter/Private/UI/DatasmithImportOptionsWindow.cpp
+++ b/Engine/Plugins/Enterprise/DatasmithImporter/Source/DatasmithImporter/Private/UI/DatasmithImportOptionsWindow.cpp
@@ -189,11 +189,11 @@ TSharedRef< SCompoundWidget > SDatasmithOptionsWindow::ConstructWarningWidget( f
 
 	if ( bIsNewer )
 	{
-		UpdateText = FText::Format( LOCTEXT( "UpdateVersion", "For best results, install Unreal Studio version {0} or later." ), FileSDKVersion );
+		UpdateText = FText::Format( LOCTEXT( "UpdateVersion_UnrealStudio", "For best results, install Unreal Studio version {0} or later." ), FileSDKVersion );
 	}
 	else
 	{
-		UpdateText = LOCTEXT( "UpdateVersion", "For best results, update your Datasmith exporter plugin if possible." );
+		UpdateText = LOCTEXT( "UpdateVersion_Datasmith", "For best results, update your Datasmith exporter plugin if possible." );
 	}
 
 	FText WarningText(
diff --git a/Engine/Plugins/Enterprise/MDLImporter/Source/MDLImporter/Private/MDLImporterFactory.cpp b/Engine/Plugins/Enterprise/MDLImporter/Source/MDLImporter/Private/MDLImporterFactory.cpp
--- a/Engine/Plugins/Enterprise/MDLImporter/Source/MDLImporter/Private/MDLImporterFactory.cpp
+++ b/Engine/Plugins/Enterprise/MDLImporter/Source/MDLImporter/Private/MDLImporterFactory.cpp
@@ -231,11 +231,11 @@ void UMDLImporterFactory::CleanUp()
 bool UMDLImporterFactory::LoadMaterials(IMdlFileImporter* Importer,
     UObject* InPackage, EObjectFlags InFlags, const FString& InFilename, const UMDLImporterOptions& InImporterOptions, FFeedbackContext& OutWarn)
 {
-	FScopedSlowTask Progress(100.0f, LOCTEXT("MDL", "Importing File ..."), true, OutWarn);
+	FScopedSlowTask Progress(100.0f, LOCTEXT("ImportingFiles", "Importing File ..."), true, OutWarn);
 	Progress.MakeDialog();
 
 	const float StartProgress = 25.f;
-	Progress.EnterProgressFrame(StartProgress, FText::Format(LOCTEXT("MDL", "Importing materials: {0} ..."), FText::FromString(InFilename)));
+	Progress.EnterProgressFrame(StartProgress, FText::Format(LOCTEXT("ImportingMaterials", "Importing materials: {0} ..."), FText::FromString(InFilename)));
 
 	bool                     bSuccess = Importer->OpenFile(InFilename, InImporterOptions);
 	if (bSuccess)
@@ -250,11 +250,11 @@ bool UMDLImporterFactory::LoadMaterials(IMdlFileImporter* Importer,
 				const float FinishProgress   = 90.f;
 				const float MaterialProgress = (FinishProgress - StartProgress) / MaterialCount;
 				Progress.EnterProgressFrame(MaterialProgress,
-				                            FText::Format(LOCTEXT("MDL", "Material distillation: {0} ..."), FText::FromString(MaterialName)));
+				                            FText::Format(LOCTEXT("MaterialDistillation", "Material distillation: {0} ..."), FText::FromString(MaterialName)));
 			}
 			else
 			{
-				Progress.EnterProgressFrame(5.f, FText::Format(LOCTEXT("MDL", "Finishing materials: {0} ..."), FText::FromString(InFilename)));
+				Progress.EnterProgressFrame(5.f, FText::Format(LOCTEXT("FinishingMaterials", "Finishing materials: {0} ..."), FText::FromString(InFilename)));
 			}
 		};
 
diff --git a/Engine/Plugins/Experimental/HairStrands/Source/HairStrandsCore/Private/GroomBuilder.cpp b/Engine/Plugins/Experimental/HairStrands/Source/HairStrandsCore/Private/GroomBuilder.cpp
--- a/Engine/Plugins/Experimental/HairStrands/Source/HairStrandsCore/Private/GroomBuilder.cpp
+++ b/Engine/Plugins/Experimental/HairStrands/Source/HairStrandsCore/Private/GroomBuilder.cpp
@@ -955,3 +955,5 @@ void FGroomBuilder::GenerateGuides(const FHairStrandsDatas& InData, float Decima
 
 	HairStrandsBuilder::BuildInternalData(OutData, false);
 }
+
+#undef LOCTEXT_NAMESPACE
diff --git a/Engine/Plugins/Experimental/MeshModelingToolset/Source/MeshModelingTools/Private/MeshSelectionTool.cpp b/Engine/Plugins/Experimental/MeshModelingToolset/Source/MeshModelingTools/Private/MeshSelectionTool.cpp
--- a/Engine/Plugins/Experimental/MeshModelingToolset/Source/MeshModelingTools/Private/MeshSelectionTool.cpp
+++ b/Engine/Plugins/Experimental/MeshModelingToolset/Source/MeshModelingTools/Private/MeshSelectionTool.cpp
@@ -898,7 +898,7 @@ void UMeshSelectionTool::DeleteSelectedTriangles()
 	ChangeSeq->AppendChange(PreviewMesh, MoveTemp(MeshChange));
 
 	// emit combined change sequence
-	GetToolManager()->EmitObjectChange(this, MoveTemp(ChangeSeq), LOCTEXT("MeshSelectionToolDelete", "Delete Faces"));
+	GetToolManager()->EmitObjectChange(this, MoveTemp(ChangeSeq), LOCTEXT("MeshSelectionToolDeleteFaces", "Delete Faces"));
 
 	OnExternalSelectionChange();
 	bHaveModifiedMesh = true;
diff --git a/Engine/Plugins/Runtime/PostSplashScreen/Source/PostSplashScreen/Private/CustomSplashScreen.cpp b/Engine/Plugins/Runtime/PostSplashScreen/Source/PostSplashScreen/Private/CustomSplashScreen.cpp
--- a/Engine/Plugins/Runtime/PostSplashScreen/Source/PostSplashScreen/Private/CustomSplashScreen.cpp
+++ b/Engine/Plugins/Runtime/PostSplashScreen/Source/PostSplashScreen/Private/CustomSplashScreen.cpp
@@ -48,4 +48,6 @@ float FCustomSplashScreen::GetAddedTickDelay()
 {
 	const float DefaultTickTime = 0.03f;
 	return FMath::Min(MaxTimeToDisplay, DefaultTickTime);
-}
\ No newline at end of file
+}
+
+#undef LOCTEXT_NAMESPACE
diff --git a/Engine/Plugins/VirtualProduction/Takes/Source/TakeRecorder/Private/Recorder/TakeRecorder.cpp b/Engine/Plugins/VirtualProduction/Takes/Source/TakeRecorder/Private/Recorder/TakeRecorder.cpp
--- a/Engine/Plugins/VirtualProduction/Takes/Source/TakeRecorder/Private/Recorder/TakeRecorder.cpp
+++ b/Engine/Plugins/VirtualProduction/Takes/Source/TakeRecorder/Private/Recorder/TakeRecorder.cpp
@@ -237,7 +237,7 @@ private:
 
 				FTimecode Timecode = FTimecode::FromFrameNumber(TotalFrames, FrameRate, FTimecode::IsDropFormatTimecodeSupported(FrameRate));
 
-				return FText::Format(LOCTEXT("RecordingText", "Recording...{0}"), FText::FromString(Timecode.ToString()));
+				return FText::Format(LOCTEXT("RecordingTimecodeText", "Recording...{0}"), FText::FromString(Timecode.ToString()));
 			}
 		}
 
diff --git a/Engine/Source/Developer/ToolMenus/Private/ToolMenus.cpp b/Engine/Source/Developer/ToolMenus/Private/ToolMenus.cpp
--- a/Engine/Source/Developer/ToolMenus/Private/ToolMenus.cpp
+++ b/Engine/Source/Developer/ToolMenus/Private/ToolMenus.cpp
@@ -1875,3 +1875,5 @@ void UToolMenus::RemoveAllCustomizations()
 {
 	CustomizedMenus.Reset();
 }
+
+#undef LOCTEXT_NAMESPACE
diff --git a/Engine/Source/Developer/TraceInsights/Private/Insights/Table/Widgets/STableTreeView.cpp b/Engine/Source/Developer/TraceInsights/Private/Insights/Table/Widgets/STableTreeView.cpp
--- a/Engine/Source/Developer/TraceInsights/Private/Insights/Table/Widgets/STableTreeView.cpp
+++ b/Engine/Source/Developer/TraceInsights/Private/Insights/Table/Widgets/STableTreeView.cpp
@@ -1217,7 +1217,7 @@ TSharedRef<SWidget> STableTreeView::GetGroupingCrumbMenuContent(const TSharedPtr
 			MenuBuilder.AddSubMenu
 			(
 				LOCTEXT("ChangeGrouping_Section", "Change To..."),
-				LOCTEXT("GroupingMenu_Add_Desc", "Change selected grouping."),
+				LOCTEXT("ChangeGrouping_Desc", "Change selected grouping."),
 				FNewMenuDelegate::CreateSP(this, &STableTreeView::BuildGroupingSubMenu_Change, CrumbGrouping),
 				false,
 				FSlateIcon()
diff --git a/Engine/Source/Runtime/ClothingSystemRuntimeCommon/Private/ClothingAsset.cpp b/Engine/Source/Runtime/ClothingSystemRuntimeCommon/Private/ClothingAsset.cpp
--- a/Engine/Source/Runtime/ClothingSystemRuntimeCommon/Private/ClothingAsset.cpp
+++ b/Engine/Source/Runtime/ClothingSystemRuntimeCommon/Private/ClothingAsset.cpp
@@ -197,9 +197,8 @@ bool UClothingAssetCommon::BindToSkeletalMesh(
 	// If we've been added to the wrong mesh
 	if(InSkelMesh != GetOuter())
 	{
-		FText Error = FText::Format(LOCTEXT(
-			"Error_WrongMesh", 
-			"Failed to bind clothing asset {0} as the provided mesh is not the owner of this asset."), 
+		FText Error = FText::Format(
+			LOCTEXT("Error_WrongMesh", "Failed to bind clothing asset {0} as the provided mesh is not the owner of this asset."), 
 			FText::FromString(GetName()));
 		Warn(Error);
 		return false;
@@ -208,12 +207,11 @@ bool UClothingAssetCommon::BindToSkeletalMesh(
 	// If we don't have clothing data
 	if(!ClothLodData.IsValidIndex(InAssetLodIndex))
 	{
-		FText Error = FText::Format(LOCTEXT(
-			"Error_NoClothingLod", 
-			"Failed to bind clothing asset {0} LOD{1} as LOD{2} does not exist."), 
+		FText Error = FText::Format(
+			LOCTEXT("Error_NoClothingLod", "Failed to bind clothing asset {0} LOD{1} as LOD{2} does not exist."), 
 			FText::FromString(GetName()), 
-			FText::AsNumber(InAssetLodIndex), 
-			FText::AsNumber(InAssetLodIndex));
+			InAssetLodIndex,
+			InAssetLodIndex);
 		Warn(Error);
 		return false;
 	}
@@ -221,9 +219,8 @@ bool UClothingAssetCommon::BindToSkeletalMesh(
 	// If we don't have a mesh
 	if(!InSkelMesh)
 	{
-		FText Error = FText::Format(LOCTEXT(
-			"Error_NoMesh", 
-			"Failed to bind clothing asset {0} as provided skel mesh does not exist."), 
+		FText Error = FText::Format(
+			LOCTEXT("Error_NoMesh", "Failed to bind clothing asset {0} as provided skel mesh does not exist."), 
 			FText::FromString(GetName()));
 		Warn(Error);
 		return false;
@@ -232,11 +229,10 @@ bool UClothingAssetCommon::BindToSkeletalMesh(
 	// If the mesh LOD index is invalid
 	if(!InSkelMesh->GetImportedModel()->LODModels.IsValidIndex(InMeshLodIndex))
 	{
-		FText Error = FText::Format(LOCTEXT(
-			"Error_InvalidMeshLOD", 
-			"Failed to bind clothing asset {0} as mesh LOD{1} does not exist."), 
+		FText Error = FText::Format(
+			LOCTEXT("Error_InvalidMeshLOD", "Failed to bind clothing asset {0} as mesh LOD{1} does not exist."), 
 			FText::FromString(GetName()), 
-			FText::AsNumber(InMeshLodIndex));
+			InMeshLodIndex);
 		Warn(Error);
 		return false;
 	}
@@ -247,13 +243,12 @@ bool UClothingAssetCommon::BindToSkeletalMesh(
 		const int32& MappedLod = LodMap[MapIndex];
 		if(MappedLod == InAssetLodIndex)
 		{
-			FText Error = FText::Format(LOCTEXT(
-				"Error_LodMapped", 
-				"Failed to bind clothing asset {0} LOD{1} as LOD{2} is already mapped to mesh LOD{3}."), 
+			FText Error = FText::Format(
+				LOCTEXT("Error_LodMapped", "Failed to bind clothing asset {0} LOD{1} as LOD{2} is already mapped to mesh LOD{3}."), 
 				FText::FromString(GetName()), 
-				FText::AsNumber(InAssetLodIndex), 
-				FText::AsNumber(InAssetLodIndex), 
-				FText::AsNumber(MapIndex));
+				InAssetLodIndex, 
+				InAssetLodIndex, 
+				MapIndex);
 			Warn(Error);
 			return false;
 		}
@@ -337,14 +332,12 @@ bool UClothingAssetCommon::BindToSkeletalMesh(
 	if(TempBoneMap.Num() > FGPUBaseSkinVertexFactory::GetMaxGPUSkinBones())
 	{
 		// Failed to apply as we've exceeded the number of bones we can skin
-		FText Error = FText::Format(LOCTEXT(
-			"Error_TooManyBones", 
-			"Failed to bind clothing asset {0} LOD{1} as this causes the section to require {2} bones. "
-			"The maximum per section is currently {3}."), 
+		FText Error = FText::Format(
+			LOCTEXT("Error_TooManyBones", "Failed to bind clothing asset {0} LOD{1} as this causes the section to require {2} bones. The maximum per section is currently {3}."), 
 			FText::FromString(GetName()), 
-			FText::AsNumber(InAssetLodIndex), 
-			FText::AsNumber(TempBoneMap.Num()), 
-			FText::AsNumber(FGPUBaseSkinVertexFactory::GetMaxGPUSkinBones()));
+			InAssetLodIndex, 
+			TempBoneMap.Num(), 
+			FGPUBaseSkinVertexFactory::GetMaxGPUSkinBones());
 		Warn(Error);
 		return false;
 	}
@@ -433,11 +426,10 @@ void UClothingAssetCommon::UnbindFromSkeletalMesh(
 	{
 		if(!Mesh->LODModels.IsValidIndex(InMeshLodIndex))
 		{
-			FText Error = FText::Format(LOCTEXT(
-				"Error_UnbindNoMeshLod", 
-				"Failed to remove clothing asset {0} from mesh LOD{1} as that LOD doesn't exist."), 
+			FText Error = FText::Format(
+				LOCTEXT("Error_UnbindNoMeshLod", "Failed to remove clothing asset {0} from mesh LOD{1} as that LOD doesn't exist."), 
 				FText::FromString(GetName()), 
-				FText::AsNumber(InMeshLodIndex));
+				InMeshLodIndex);
 			Warn(Error);
 
 			return;
diff --git a/Engine/Source/Runtime/CoreUObject/Private/UObject/ScriptCore.cpp b/Engine/Source/Runtime/CoreUObject/Private/UObject/ScriptCore.cpp
--- a/Engine/Source/Runtime/CoreUObject/Private/UObject/ScriptCore.cpp
+++ b/Engine/Source/Runtime/CoreUObject/Private/UObject/ScriptCore.cpp
@@ -2001,7 +2001,7 @@ DEFINE_FUNCTION(UObject::execClassSparseDataVariable)
 
 	if (VarProperty == nullptr || P_THIS->GetSparseClassDataStruct() == nullptr)
 	{
-		FBlueprintExceptionInfo ExceptionInfo(EBlueprintExceptionType::AccessViolation, FText::Format(LOCTEXT("MissingProperty", "Attempted to access missing sparse property '{0}' {1}, {2}. If this is a packaged/cooked build, are you attempting to use an editor-only property?"), FText::FromString(GetNameSafe(VarProperty)), FText::FromString(GetNameSafe(P_THIS->GetSparseClassDataStruct())), FText::FromString(GetNameSafe(VarProperty ? (UClass*)VarProperty->GetOuter() : nullptr))));
+		FBlueprintExceptionInfo ExceptionInfo(EBlueprintExceptionType::AccessViolation, FText::Format(LOCTEXT("MissingSparseProperty", "Attempted to access missing sparse property '{0}' {1}, {2}. If this is a packaged/cooked build, are you attempting to use an editor-only property?"), FText::FromString(GetNameSafe(VarProperty)), FText::FromString(GetNameSafe(P_THIS->GetSparseClassDataStruct())), FText::FromString(GetNameSafe(VarProperty ? (UClass*)VarProperty->GetOuter() : nullptr))));
 		FBlueprintCoreDelegates::ThrowScriptException(P_THIS, Stack, ExceptionInfo);
 
 		Stack.MostRecentPropertyAddress = nullptr;
diff --git a/Engine/Source/Runtime/Launch/Private/LaunchEngineLoop.cpp b/Engine/Source/Runtime/Launch/Private/LaunchEngineLoop.cpp
--- a/Engine/Source/Runtime/Launch/Private/LaunchEngineLoop.cpp
+++ b/Engine/Source/Runtime/Launch/Private/LaunchEngineLoop.cpp
@@ -2279,7 +2279,7 @@ int32 FEngineLoop::PreInitPreStartupScreen(const TCHAR* CmdLine)
 
 	FEmbeddedCommunication::ForceTick(3);
 
-	FScopedSlowTask SlowTask(50, NSLOCTEXT("EngineLoop", "EngineLoop_Initializing", "PreInitPreStartupScreen..."));
+	FScopedSlowTask SlowTask(50, NSLOCTEXT("EngineLoop", "EngineLoop_Initializing_PreInitPreStartupScreen", "PreInitPreStartupScreen..."));
 
 	SlowTask.EnterProgressFrame(10);
 
@@ -2514,7 +2514,7 @@ int32 FEngineLoop::PreInitPreStartupScreen(const TCHAR* CmdLine)
 
 int32 FEngineLoop::PreInitPostStartupScreen(const TCHAR* CmdLine)
 {
-	FScopedSlowTask SlowTask(50, NSLOCTEXT("EngineLoop", "EngineLoop_Initializing", "PreInitPostStartupScreen..."));
+	FScopedSlowTask SlowTask(50, NSLOCTEXT("EngineLoop", "EngineLoop_Initializing_PreInitPostStartupScreen", "PreInitPostStartupScreen..."));
 	FScopeCycleCounter CycleCount_AfterStats(GET_STATID(STAT_FEngineLoop_PreInitPostStartupScreen_AfterStats));
 
 	// Restore PreInitContext
