- Chaos triangle mesh face normals are now calculated assuming CCW winding in left handed coordinates. This default is made explicit in the comments to prevent future confusion
- Chaos normal debug rendering now the correct way around for newly imported assets
- Dead code removed from ClothingMeshUtils
- Comments added in code when non default winding order is used for face normal calculations, giving opposite results (like calculation of barycentric coordinates)

diff --git a/Engine/Plugins/Experimental/ChaosCloth/Source/Chaos/Private/ChaosCloth/ChaosClothingSimulation.cpp b/Engine/Plugins/Experimental/ChaosCloth/Source/Chaos/Private/ChaosCloth/ChaosClothingSimulation.cpp
--- a/Engine/Plugins/Experimental/ChaosCloth/Source/Chaos/Private/ChaosCloth/ChaosClothingSimulation.cpp
+++ b/Engine/Plugins/Experimental/ChaosCloth/Source/Chaos/Private/ChaosCloth/ChaosClothingSimulation.cpp
@@ -1134,7 +1134,7 @@ void ClothingSimulation::GetSimulationData(
         {
 			const uint32 LocalIndex = j - VertexDomain[0];
             Data.Positions[LocalIndex] = Evolution->Particles().X(j);
-            Data.Normals[LocalIndex] = PointNormals[i][LocalIndex];
+            Data.Normals[LocalIndex] = -PointNormals[i][LocalIndex]; // Note the Normals are inverted due to how barycentric coordinates are calculated (see GetPointBaryAndDist in ClothingMeshUtils.cpp)
 		}
     }
 }
diff --git a/Engine/Source/Runtime/ClothingSystemRuntimeCommon/Private/Utils/ClothingMeshUtils.cpp b/Engine/Source/Runtime/ClothingSystemRuntimeCommon/Private/Utils/ClothingMeshUtils.cpp
--- a/Engine/Source/Runtime/ClothingSystemRuntimeCommon/Private/Utils/ClothingMeshUtils.cpp
+++ b/Engine/Source/Runtime/ClothingSystemRuntimeCommon/Private/Utils/ClothingMeshUtils.cpp
@@ -279,19 +279,13 @@ namespace ClothingMeshUtils
 		}
 	}
 
+	// TODO: Vertex normals are not used at present, a future improved algorithm might however
 	FVector4 GetPointBaryAndDist(const FVector& A, const FVector& B, const FVector& C, const FVector& NA, const FVector& NB, const FVector& NC, const FVector& Point)
 	{
 		FPlane TrianglePlane(A, B, C);
 		const FVector PointOnTriPlane = FVector::PointPlaneProject(Point, TrianglePlane);
 		const FVector BaryCoords = FMath::ComputeBaryCentric2D(PointOnTriPlane, A, B, C);
-		const FVector NormalAtPoint = TrianglePlane;
-		FVector TriPointToVert = Point - PointOnTriPlane;
-		TriPointToVert = TriPointToVert.ProjectOnTo(NormalAtPoint);
-		float Dist = TriPointToVert.Size();
-
-		float Sign = TrianglePlane.PlaneDot(Point) < 0.0f ? -1.0f : 1.0f;
-
-		return FVector4(BaryCoords, TrianglePlane.PlaneDot(Point));
+		return FVector4(BaryCoords, TrianglePlane.PlaneDot(Point)); // Note: The normal of the plane points away from the Clockwise face (instead of the counter clockwise face) in Left Handed Coordinates (This is why we need to invert the normals later on when before sending it to the shader)
 	}
 
 	void GenerateEmbeddedPositions(const ClothMeshDesc& SourceMesh, TArrayView<const FVector> Positions, TArray<FVector4>& OutEmbeddedPositions, TArray<int32>& OutSourceIndices)
diff --git a/Engine/Source/Runtime/ClothingSystemRuntimeNv/Private/ClothingSimulationNv.cpp b/Engine/Source/Runtime/ClothingSystemRuntimeNv/Private/ClothingSimulationNv.cpp
--- a/Engine/Source/Runtime/ClothingSystemRuntimeNv/Private/ClothingSimulationNv.cpp
+++ b/Engine/Source/Runtime/ClothingSystemRuntimeNv/Private/ClothingSimulationNv.cpp
@@ -872,6 +872,8 @@ void FClothingSimulationNv::ComputePhysicalMeshNormals(FClothingActorNv &Actor)
 
 	// Using the face normals, calculate normals. These will not be normalized as we're adding together
 	// normals for every shared face a vert has. We'll normalize later
+	// Note: The normal points away from the Clockwise face (instead of the counter clockwise face) in Left Handed Coordinates
+	// this is to be consistent with how barycentric coordinate distances are calculated
 	for(uint32 BaseIndex = 0; BaseIndex < NumIndices; BaseIndex += 3)
 	{
 		const FVector A = P2UVector(Particles[Indices[BaseIndex]]);
diff --git a/Engine/Source/Runtime/Experimental/Chaos/Private/Chaos/TriangleMesh.cpp b/Engine/Source/Runtime/Experimental/Chaos/Private/Chaos/TriangleMesh.cpp
--- a/Engine/Source/Runtime/Experimental/Chaos/Private/Chaos/TriangleMesh.cpp
+++ b/Engine/Source/Runtime/Experimental/Chaos/Private/Chaos/TriangleMesh.cpp
@@ -237,6 +237,8 @@ TArray<TVector<T, 3>> TTriangleMesh<T>::GetFaceNormals(const TArrayView<const TV
 	return Normals;
 }
 
+// Note:	This function assumes Counter Clockwise triangle windings in a Left Handed coordinate system
+//			If this is not the case the returned face normals may need to be inverted
 template<class T>
 void TTriangleMesh<T>::GetFaceNormals(TArray<TVector<T, 3>>& Normals, const TArrayView<const TVector<T, 3>>& Points, const bool ReturnEmptyOnError) const
 {
@@ -247,7 +249,7 @@ void TTriangleMesh<T>::GetFaceNormals(TArray<TVector<T, 3>>& Normals, const TArr
 		{
 			TVector<T, 3> p10 = Points[Tri[1]] - Points[Tri[0]];
 			TVector<T, 3> p20 = Points[Tri[2]] - Points[Tri[0]];
-			TVector<T, 3> Cross = TVector<T, 3>::CrossProduct(p10, p20);
+			TVector<T, 3> Cross = TVector<T, 3>::CrossProduct(p20, p10);
 			const T Size2 = Cross.SizeSquared();
 			if (Size2 < SMALL_NUMBER)
 			{
@@ -268,7 +270,7 @@ void TTriangleMesh<T>::GetFaceNormals(TArray<TVector<T, 3>>& Normals, const TArr
 		{
 			TVector<T, 3> p10 = Points[Tri[1]] - Points[Tri[0]];
 			TVector<T, 3> p20 = Points[Tri[2]] - Points[Tri[0]];
-			TVector<T, 3> Cross = TVector<T, 3>::CrossProduct(p10, p20);
+			TVector<T, 3> Cross = TVector<T, 3>::CrossProduct(p20, p10);
 			Normals.Add(Cross.GetSafeNormal());
 		}
 	}
diff --git a/Engine/Source/Runtime/Experimental/Chaos/Public/Chaos/TriangleMesh.h b/Engine/Source/Runtime/Experimental/Chaos/Public/Chaos/TriangleMesh.h
--- a/Engine/Source/Runtime/Experimental/Chaos/Public/Chaos/TriangleMesh.h
+++ b/Engine/Source/Runtime/Experimental/Chaos/Public/Chaos/TriangleMesh.h
@@ -126,6 +126,9 @@ namespace Chaos
 		TArray<Chaos::TVector<int32, 2>> GetUniqueAdjacentPoints();
 		TArray<Chaos::TVector<int32, 4>> GetUniqueAdjacentElements();
 
+		/** The GetFaceNormals functions assume Counter Clockwise triangle windings in a Left Handed coordinate system
+			If this is not the case the returned face normals may be inverted
+		*/
 		TArray<TVector<T, 3>> GetFaceNormals(const TArrayView<const TVector<T, 3>>& Points, const bool ReturnEmptyOnError = true) const;
 		void GetFaceNormals(TArray<TVector<T, 3>>& Normals, const TArrayView<const TVector<T, 3>>& Points, const bool ReturnEmptyOnError = true) const;
 		/** Deprecated. Use TArrayView version. */
