Fix CIS after Clothing Asset changes.

diff --git a/Engine/Source/Runtime/ClothingSystemRuntimeCommon/Private/ClothLODData_Legacy.cpp b/Engine/Source/Runtime/ClothingSystemRuntimeCommon/Private/ClothLODData_Legacy.cpp
--- a/Engine/Source/Runtime/ClothingSystemRuntimeCommon/Private/ClothLODData_Legacy.cpp
+++ b/Engine/Source/Runtime/ClothingSystemRuntimeCommon/Private/ClothLODData_Legacy.cpp
@@ -14,10 +14,12 @@ FClothParameterMask_Legacy::FClothParameterMask_Legacy()
 
 void FClothParameterMask_Legacy::MigrateTo(FPointWeightMap* Weights) const
 {
+	Weights->Values = Values;
+#if WITH_EDITORONLY_DATA
 	Weights->Name = MaskName;
 	Weights->CurrentTarget = static_cast<uint8>(CurrentTarget);
-	Weights->Values = Values;
 	Weights->bEnabled = bEnabled;
+#endif
 }
 
 bool FClothLODData_Legacy::Serialize(FArchive& Ar)
