Fix CIS missing include after Clothing Asset changes.

diff --git a/Engine/Source/Runtime/ClothingSystemRuntimeCommon/Private/ClothingAsset.cpp b/Engine/Source/Runtime/ClothingSystemRuntimeCommon/Private/ClothingAsset.cpp
--- a/Engine/Source/Runtime/ClothingSystemRuntimeCommon/Private/ClothingAsset.cpp
+++ b/Engine/Source/Runtime/ClothingSystemRuntimeCommon/Private/ClothingAsset.cpp
@@ -6,6 +6,7 @@
 #include "ClothConfig.h"
 
 #include "Utils/ClothingMeshUtils.h"
+#include "Features/IModularFeatures.h"
 
 #if WITH_EDITOR
 #include "Engine/SkeletalMesh.h"
