Fix missing Chaos Cloth config creation when reloading NvCloth assets made post Ryan's refactor.
Added a new migration path Nv Cloth config to Chaos Cloth config using the legacy config structure.

diff --git a/Engine/Source/Runtime/ClothingSystemRuntimeCommon/Private/ClothingAsset.cpp b/Engine/Source/Runtime/ClothingSystemRuntimeCommon/Private/ClothingAsset.cpp
--- a/Engine/Source/Runtime/ClothingSystemRuntimeCommon/Private/ClothingAsset.cpp
+++ b/Engine/Source/Runtime/ClothingSystemRuntimeCommon/Private/ClothingAsset.cpp
@@ -850,6 +850,9 @@ void UClothingAssetCommon::PostLoad()
 		SetClothConfig(ClothSharedSimConfig_DEPRECATED);
 		ClothSharedSimConfig_DEPRECATED = nullptr;
 	}
+
+	// Add any missing configs for the available cloth factories, and try to migrate them from any existing one
+	AddClothConfigs();
 }
 
 void UClothingAssetCommon::Serialize(FArchive& Ar)
@@ -859,6 +862,50 @@ void UClothingAssetCommon::Serialize(FArchive& Ar)
 	Ar.UsingCustomVersion(FClothingAssetCustomVersion::GUID);
 }
 
+void UClothingAssetCommon::AddClothConfigs()
+{
+	const TArray<IClothingSimulationFactoryClassProvider*> ClassProviders =
+		IModularFeatures::Get().GetModularFeatureImplementations<IClothingSimulationFactoryClassProvider>(IClothingSimulationFactoryClassProvider::FeatureName);
+
+	for (IClothingSimulationFactoryClassProvider* Provider : ClassProviders)
+	{
+		check(Provider);
+		if (UClass* const ClothingSimulationFactory = *TSubclassOf<class UClothingSimulationFactory>(Provider->GetClothingSimulationFactoryClass()))
+		{
+			const TSubclassOf<UClothConfigBase> ClothConfigClass = ClothingSimulationFactory->GetDefaultObject<UClothingSimulationFactory>()->GetClothConfigClass();
+			const FName ClothConfigName = ClothConfigClass->GetFName();
+			if (!ClothConfigs.Find(ClothConfigName))
+			{
+				// Create new config object
+				UClothConfigBase* const ClothConfig = NewObject<UClothConfigBase>(this, ClothConfigClass, ClothConfigClass->GetFName(), RF_Transactional);
+				check(ClothConfig);
+
+				// Use the legacy config struct to try find a common config as an acceptable migration source
+				// This code can be removed once the legacy code is removed, although this will then prevent
+				// migration from compatible config sources
+				if (UClothConfigCommon* const ClothConfigCommon = Cast<UClothConfigCommon>(ClothConfig))
+				{
+					for (TPair<FName, UClothConfigBase*> ClothConfigPair : ClothConfigs)
+					{
+						if (const UClothConfigCommon* SourceConfig = Cast<UClothConfigCommon>(ClothConfigPair.Value))
+						{
+							FClothConfig_Legacy LegacyConfig;
+							if (SourceConfig->MigrateTo(LegacyConfig))
+							{
+								ClothConfigCommon->MigrateFrom(LegacyConfig);
+								break;
+							}
+						}
+					}
+				}
+
+				// Add the new config
+				ClothConfigs.Add(ClothConfigName, ClothConfig);
+			}
+		}
+	}
+}
+
 #if WITH_EDITOR
 
 void UClothingAssetCommon::InvalidateCachedData()
diff --git a/Engine/Source/Runtime/ClothingSystemRuntimeCommon/Public/ClothConfig.h b/Engine/Source/Runtime/ClothingSystemRuntimeCommon/Public/ClothConfig.h
--- a/Engine/Source/Runtime/ClothingSystemRuntimeCommon/Public/ClothConfig.h
+++ b/Engine/Source/Runtime/ClothingSystemRuntimeCommon/Public/ClothConfig.h
@@ -17,4 +17,11 @@ public:
 
 	/** Migrate from the legacy FClothConfig structure. */
 	virtual void MigrateFrom(const FClothConfig_Legacy&) {}
+
+	/**
+	 * Migrate to the legacy FClothConfig structure.
+	 * Useful for converting configs that are compatible with this legacy structure.
+	 * @return true when the migration is possible, false otherwise.
+	 */
+	virtual bool MigrateTo(FClothConfig_Legacy&) const { return false; }
 };
diff --git a/Engine/Source/Runtime/ClothingSystemRuntimeCommon/Public/ClothingAsset.h b/Engine/Source/Runtime/ClothingSystemRuntimeCommon/Public/ClothingAsset.h
--- a/Engine/Source/Runtime/ClothingSystemRuntimeCommon/Public/ClothingAsset.h
+++ b/Engine/Source/Runtime/ClothingSystemRuntimeCommon/Public/ClothingAsset.h
@@ -257,6 +257,10 @@ public:
 	TArray<FClothLODData_Legacy> LodData_DEPRECATED;
 
 private:
+
+	// Add any missing cloth configs, for example after reloading clothing assets
+	void AddClothConfigs();
+
 #if WITH_EDITOR
 	// Helper functions used in PostPropertyChangeCb
 	void ReregisterComponentsUsingClothing();
diff --git a/Engine/Source/Runtime/ClothingSystemRuntimeNv/Private/ClothConfigNv.cpp b/Engine/Source/Runtime/ClothingSystemRuntimeNv/Private/ClothConfigNv.cpp
--- a/Engine/Source/Runtime/ClothingSystemRuntimeNv/Private/ClothConfigNv.cpp
+++ b/Engine/Source/Runtime/ClothingSystemRuntimeNv/Private/ClothConfigNv.cpp
@@ -19,6 +19,14 @@ void FClothConstraintSetupNv::MigrateFrom(const FClothConstraintSetup_Legacy& Se
 	CompressionLimit = Setup.CompressionLimit;
 }
 
+void FClothConstraintSetupNv::MigrateTo(FClothConstraintSetup_Legacy& Setup) const
+{
+	Setup.Stiffness = Stiffness;
+	Setup.StiffnessMultiplier = StiffnessMultiplier;
+	Setup.StretchLimit = StretchLimit;
+	Setup.CompressionLimit = CompressionLimit;
+}
+
 UClothConfigNv::UClothConfigNv()
 	: ClothingWindMethod(EClothingWindMethodNv::Legacy)
 	, SelfCollisionRadius(0.0f)
@@ -98,6 +106,40 @@ void UClothConfigNv::MigrateFrom(const FClothConfig_Legacy& ClothConfig)
 	AnimDriveDamperStiffness = ClothConfig.AnimDriveDamperStiffness;
 }
 
+bool UClothConfigNv::MigrateTo(FClothConfig_Legacy& ClothConfig) const
+{
+	ClothConfig.WindMethod = static_cast<EClothingWindMethod_Legacy>(ClothingWindMethod);
+	VerticalConstraint.MigrateTo(ClothConfig.VerticalConstraintConfig);
+	HorizontalConstraint.MigrateTo(ClothConfig.HorizontalConstraintConfig);
+	BendConstraint.MigrateTo(ClothConfig.BendConstraintConfig);
+	ShearConstraint.MigrateTo(ClothConfig.ShearConstraintConfig);
+	ClothConfig.SelfCollisionRadius = SelfCollisionRadius;
+	ClothConfig.SelfCollisionStiffness = SelfCollisionStiffness;
+	ClothConfig.SelfCollisionCullScale = SelfCollisionCullScale;
+	ClothConfig.Damping = Damping;
+	ClothConfig.Friction = Friction;
+	ClothConfig.WindDragCoefficient = WindDragCoefficient;
+	ClothConfig.WindLiftCoefficient = WindLiftCoefficient;
+	ClothConfig.LinearDrag = LinearDrag;
+	ClothConfig.AngularDrag = AngularDrag;
+	ClothConfig.LinearInertiaScale = LinearInertiaScale;
+	ClothConfig.AngularInertiaScale = AngularInertiaScale;
+	ClothConfig.CentrifugalInertiaScale = CentrifugalInertiaScale;
+	ClothConfig.SolverFrequency = SolverFrequency;
+	ClothConfig.StiffnessFrequency = StiffnessFrequency;
+	ClothConfig.GravityScale = GravityScale;
+	ClothConfig.GravityOverride = GravityOverride;
+	ClothConfig.bUseGravityOverride = bUseGravityOverride;
+	ClothConfig.TetherStiffness = TetherStiffness;
+	ClothConfig.TetherLimit = TetherLimit;
+	ClothConfig.CollisionThickness = CollisionThickness;
+	ClothConfig.AnimDriveSpringStiffness = AnimDriveSpringStiffness;
+	ClothConfig.AnimDriveDamperStiffness = AnimDriveDamperStiffness;
+
+	return true;
+}
+
+
 bool UClothConfigNv::UseSelfCollisions() const
 {
 	return SelfCollisionRadius > 0.0f && SelfCollisionStiffness > 0.0f;
diff --git a/Engine/Source/Runtime/ClothingSystemRuntimeNv/Public/ClothConfigNv.h b/Engine/Source/Runtime/ClothingSystemRuntimeNv/Public/ClothConfigNv.h
--- a/Engine/Source/Runtime/ClothingSystemRuntimeNv/Public/ClothConfigNv.h
+++ b/Engine/Source/Runtime/ClothingSystemRuntimeNv/Public/ClothConfigNv.h
@@ -14,8 +14,12 @@ struct FClothConstraintSetupNv
 
 	FClothConstraintSetupNv();
 
+	// Migrate from the legacy FClothConstraintSetup structure.
 	void MigrateFrom(const FClothConstraintSetup_Legacy& Setup);
 
+	// Migrate to the legacy FClothConstraintSetup structure.
+	void MigrateTo(FClothConstraintSetup_Legacy& Setup) const;
+
 	// How stiff this constraint is, this affects how closely it will follow the desired position
 	UPROPERTY(EditAnywhere, Category=Constraint)
 	float Stiffness;
@@ -65,6 +69,9 @@ public:
 	// Migrate from the legacy FClothConfig structure.
 	virtual void MigrateFrom(const FClothConfig_Legacy& ClothConfig) override;
 
+	// Migrate to the legacy FClothConfig structure.
+	virtual bool MigrateTo(FClothConfig_Legacy& ClothConfig) const override;
+
 	// Return whether self collision is enabled for this config.
 	UE_DEPRECATED(4.25, "This function is deprecated. Please use NeedsSelfCollisionIndices or UseSelfCollisions instead.")
 	virtual bool HasSelfCollision() const override
