Chaos Cloth - Fix ReferenceBoneIndex calculations and mapping.
- Fix the ClothLodData and PhysicalMeshData migration code that was missing ConditionalPostLoad.

diff --git a/Engine/Source/Runtime/ClothingSystemRuntimeCommon/Private/ClothLODData.cpp b/Engine/Source/Runtime/ClothingSystemRuntimeCommon/Private/ClothLODData.cpp
--- a/Engine/Source/Runtime/ClothingSystemRuntimeCommon/Private/ClothLODData.cpp
+++ b/Engine/Source/Runtime/ClothingSystemRuntimeCommon/Private/ClothLODData.cpp
@@ -53,6 +53,7 @@ void UClothLODDataCommon::PostLoad()
 
 	if (PhysicalMeshData_DEPRECATED)
 	{
+		PhysicalMeshData_DEPRECATED->ConditionalPostLoad();  // Makes sure the UObject has finished loading
 		ClothPhysicalMeshData.MigrateFrom(PhysicalMeshData_DEPRECATED);
 		PhysicalMeshData_DEPRECATED = nullptr;
 	}
diff --git a/Engine/Source/Runtime/ClothingSystemRuntimeCommon/Private/ClothingAsset.cpp b/Engine/Source/Runtime/ClothingSystemRuntimeCommon/Private/ClothingAsset.cpp
--- a/Engine/Source/Runtime/ClothingSystemRuntimeCommon/Private/ClothingAsset.cpp
+++ b/Engine/Source/Runtime/ClothingSystemRuntimeCommon/Private/ClothingAsset.cpp
@@ -599,12 +599,8 @@ void UClothingAssetCommon::CalculateReferenceBoneIndex()
 
 	if(OwnerMesh)
 	{
-		FReferenceSkeleton& RefSkel = OwnerMesh->RefSkeleton;
 		// First build a list per used bone for it's path to root
-		const int32 NumUsedBones = UsedBoneIndices.Num();
-
-		// List of actually weighted (not just used) bones
-		TArray<int32> WeightedBones;
+		TArray<int32> WeightedBones;  // List of actually weighted (not just used) bones
 
 		for(UClothLODDataCommon* CurLod : ClothLodData)
 		{
@@ -616,7 +612,9 @@ void UClothingAssetCommon::CalculateReferenceBoneIndex()
 				{
 					if(VertBoneData.BoneWeights[InfluenceIndex] > SMALL_NUMBER)
 					{
-						WeightedBones.AddUnique(VertBoneData.BoneIndices[InfluenceIndex]);
+						const int32 UnmappedBoneIndex = VertBoneData.BoneIndices[InfluenceIndex];
+						check(UsedBoneIndices.IsValidIndex(UnmappedBoneIndex));
+						WeightedBones.AddUnique(UsedBoneIndices[UnmappedBoneIndex]);
 					}
 					else
 					{
@@ -631,6 +629,7 @@ void UClothingAssetCommon::CalculateReferenceBoneIndex()
 		PathsToRoot.Reserve(NumWeightedBones);
 		
 		// Compute paths to the root bone
+		const FReferenceSkeleton& RefSkel = OwnerMesh->RefSkeleton;
 		for(int32 WeightedBoneIndex = 0; WeightedBoneIndex < NumWeightedBones; ++WeightedBoneIndex)
 		{
 			PathsToRoot.AddDefaulted();
@@ -717,6 +716,10 @@ void UClothingAssetCommon::BuildSelfCollisionData()
 void UClothingAssetCommon::PostLoad()
 {
 	Super::PostLoad();
+	for (UClothLODDataCommon* Lod : ClothLodData)
+	{
+		Lod->ConditionalPostLoad();
+	}
 
 	const int32 AnimPhysCustomVersion = GetLinkerCustomVersion(FAnimPhysObjectVersion::GUID);
 	const int32 ClothingCustomVersion = GetLinkerCustomVersion(FClothingAssetCustomVersion::GUID);
@@ -810,12 +813,6 @@ void UClothingAssetCommon::PostLoad()
 	}
 #endif
 
-	// After fixing the content, we are ready to call functions that rely on it
-	BuildSelfCollisionData();
-#if WITH_EDITORONLY_DATA
-	CalculateReferenceBoneIndex();
-#endif
-
 	// Migrate simulation dependent config parameters to the new config map
 	if (ClothSimConfig_DEPRECATED)
 	{
@@ -848,6 +845,12 @@ void UClothingAssetCommon::PostLoad()
 		SetClothConfig(ClothSharedSimConfig_DEPRECATED);
 		ClothSharedSimConfig_DEPRECATED = nullptr;
 	}
+
+	// After fixing the content, we are ready to call functions that rely on it
+	BuildSelfCollisionData();
+#if WITH_EDITORONLY_DATA
+	CalculateReferenceBoneIndex();
+#endif
 }
 
 void UClothingAssetCommon::Serialize(FArchive& Ar)
