Fix CIS error 'MaxDistanceScale' is not a base or member.

diff --git a/Engine/Source/Runtime/ClothingSystemRuntimeCommon/Private/ClothingSimulation.cpp b/Engine/Source/Runtime/ClothingSystemRuntimeCommon/Private/ClothingSimulation.cpp
--- a/Engine/Source/Runtime/ClothingSystemRuntimeCommon/Private/ClothingSimulation.cpp
+++ b/Engine/Source/Runtime/ClothingSystemRuntimeCommon/Private/ClothingSimulation.cpp
@@ -20,6 +20,7 @@ FClothingSimulationContextCommon::FClothingSimulationContextCommon()
 	, WindVelocity(FVector::ZeroVector)
 	, DeltaSeconds(0.f)
 	, TeleportMode(EClothingTeleportMode::None)
+	, MaxDistanceScale(1.f)
 {}
 
 FClothingSimulationContextCommon::~FClothingSimulationContextCommon()
diff --git a/Engine/Source/Runtime/ClothingSystemRuntimeNv/Private/ClothingSimulationNv.cpp b/Engine/Source/Runtime/ClothingSystemRuntimeNv/Private/ClothingSimulationNv.cpp
--- a/Engine/Source/Runtime/ClothingSystemRuntimeNv/Private/ClothingSimulationNv.cpp
+++ b/Engine/Source/Runtime/ClothingSystemRuntimeNv/Private/ClothingSimulationNv.cpp
@@ -41,7 +41,6 @@ DECLARE_CYCLE_STAT(TEXT("Update Anim Drive"), STAT_NvClothUpdateAnimDrive, STATG
 FClothingSimulationContextNv::FClothingSimulationContextNv()
 	: PredictedLod(0)
 	, WindAdaption(0.0f)
-	, MaxDistanceScale(1.f)
 {
 }
 
