Chaos Cloth - Expose friction coefficient to try limit sliding around collision volumes.
- Add per cloth thickness.
- Add shared config to config migration path.
- Fix cloth asset PostLoad migration issues.

diff --git a/Engine/Plugins/Experimental/ChaosCloth/Source/Chaos/Private/ChaosCloth/ChaosClothConfig.cpp b/Engine/Plugins/Experimental/ChaosCloth/Source/Chaos/Private/ChaosCloth/ChaosClothConfig.cpp
--- a/Engine/Plugins/Experimental/ChaosCloth/Source/Chaos/Private/ChaosCloth/ChaosClothConfig.cpp
+++ b/Engine/Plugins/Experimental/ChaosCloth/Source/Chaos/Private/ChaosCloth/ChaosClothConfig.cpp
@@ -66,11 +66,30 @@ void UChaosClothConfig::MigrateFrom(const FClothConfig_Legacy& ClothConfig)
 	const float Damping = (ClothConfig.Damping.X + ClothConfig.Damping.Y + ClothConfig.Damping.Z) / 3.f;
 	DampingCoefficient = FMath::Clamp(Damping * Damping * 0.95f, 0.f, 1.f);  // Nv Cloth seems to have a different damping formulation.
 
+	CollisionThickness = FMath::Clamp(ClothConfig.CollisionThickness, 0.f, 1000.f);
+
 	LinearVelocityScale = ClothConfig.LinearInertiaScale;
 	const FVector AngularInertiaScale = ClothConfig.AngularInertiaScale * ClothConfig.CentrifugalInertiaScale;
 	AngularVelocityScale = (AngularInertiaScale.X + AngularInertiaScale.Y + AngularInertiaScale.Z) / 3.f;
 }
 
+void UChaosClothConfig::MigrateFrom(const UClothSharedConfigCommon* ClothSharedConfig)
+{
+	if (const UChaosClothSharedSimConfig* const ChaosClothSharedSimConfig = Cast<UChaosClothSharedSimConfig>(ClothSharedConfig))
+	{
+		const int32 ChaosClothConfigCustomVersion = GetLinkerCustomVersion(FChaosClothConfigCustomVersion::GUID);
+
+		if (ChaosClothConfigCustomVersion < FChaosClothConfigCustomVersion::AddDampingThicknessMigration)
+		{
+			if (ChaosClothSharedSimConfig->bUseDampingOverride_DEPRECATED)
+			{
+				DampingCoefficient = ChaosClothSharedSimConfig->Damping_DEPRECATED;
+			}
+			CollisionThickness = ChaosClothSharedSimConfig->CollisionThickness_DEPRECATED;
+		}
+	}
+}
+
 void UChaosClothConfig::Serialize(FArchive& Ar)
 {
 	Super::Serialize(Ar);
@@ -100,15 +119,13 @@ void UChaosClothSharedSimConfig::MigrateFrom(const FClothConfig_Legacy& ClothCon
 
 	SelfCollisionThickness = FMath::Clamp(ClothConfig.SelfCollisionRadius, 0.f, 1000.f);
 
-	CollisionThickness = FMath::Clamp(ClothConfig.CollisionThickness, 0.f, 1000.f);
-
 	bUseGravityOverride = ClothConfig.bUseGravityOverride;
 
 	GravityScale = ClothConfig.GravityScale;
 
 	Gravity = ClothConfig.GravityOverride;
 
-	bUseDampingOverride = false;  // Damping is migrated to per cloth configs
+	bUseDampingOverride_DEPRECATED = false;  // Damping is migrated to per cloth configs
 }
 
 void UChaosClothSharedSimConfig::Serialize(FArchive& Ar)
diff --git a/Engine/Plugins/Experimental/ChaosCloth/Source/Chaos/Private/ChaosCloth/ChaosClothConfigCustomVersion.h b/Engine/Plugins/Experimental/ChaosCloth/Source/Chaos/Private/ChaosCloth/ChaosClothConfigCustomVersion.h
--- a/Engine/Plugins/Experimental/ChaosCloth/Source/Chaos/Private/ChaosCloth/ChaosClothConfigCustomVersion.h
+++ b/Engine/Plugins/Experimental/ChaosCloth/Source/Chaos/Private/ChaosCloth/ChaosClothConfigCustomVersion.h
@@ -14,6 +14,8 @@ struct FChaosClothConfigCustomVersion
 		BeforeCustomVersionWasAdded = 0,
 		// Update drag default to better preserve legacy behavior
 		UpdateDragDefault = 1,
+		// Added damping and collision thickness per cloth
+		AddDampingThicknessMigration = 2,
 
 		// -----<new versions can be added above this line>-------------------------------------------------
 		VersionPlusOne,
diff --git a/Engine/Plugins/Experimental/ChaosCloth/Source/Chaos/Private/ChaosCloth/ChaosClothingSimulation.cpp b/Engine/Plugins/Experimental/ChaosCloth/Source/Chaos/Private/ChaosCloth/ChaosClothingSimulation.cpp
--- a/Engine/Plugins/Experimental/ChaosCloth/Source/Chaos/Private/ChaosCloth/ChaosClothingSimulation.cpp
+++ b/Engine/Plugins/Experimental/ChaosCloth/Source/Chaos/Private/ChaosCloth/ChaosClothingSimulation.cpp
@@ -299,15 +299,10 @@ void ClothingSimulation::CreateActor(USkeletalMeshComponent* InOwnerComponent, U
 
 	AddConstraints(ChaosClothSimConfig, PhysMesh, InSimDataIndex);
 
-	// Set damping
-	if (ClothSharedSimConfig && ClothSharedSimConfig->bUseDampingOverride)
-	{
-		Evolution->SetDamping(InSimDataIndex, ClothSharedSimConfig->Damping);
-	}
-	else
-	{
-		Evolution->SetDamping(InSimDataIndex, ChaosClothSimConfig->DampingCoefficient);
-	}
+	// Set this cloth damping, collision thickness, friction
+	Evolution->SetDamping(ChaosClothSimConfig->DampingCoefficient, InSimDataIndex);
+	Evolution->SetCollisionThickness(ChaosClothSimConfig->CollisionThickness, InSimDataIndex);
+	Evolution->SetCoefficientOfFriction(ChaosClothSimConfig->FrictionCoefficient, InSimDataIndex);
 
 	// Add velocity field
 	auto GetVelocity = [this](const TVector<float, 3>&)->TVector<float, 3>
@@ -380,7 +375,6 @@ void ClothingSimulation::UpdateSimulationFromSharedSimConfig()
 		// Now set all the common parameters on the simulation
 		Evolution->SetIterations(ClothSharedSimConfig->IterationCount);
 		Evolution->SetSelfCollisionThickness(ClothSharedSimConfig->SelfCollisionThickness);
-		Evolution->SetCollisionThickness(ClothSharedSimConfig->CollisionThickness);
 	}
 }
 
@@ -1662,15 +1656,10 @@ void ClothingSimulation::RefreshClothConfig()
 				LinearDeltaRatios[SimDataIndex] = FVector::OneVector - ChaosClothConfig->LinearVelocityScale.BoundToBox(FVector::ZeroVector, FVector::OneVector);
 				AngularDeltaRatios[SimDataIndex] = 1.f - FMath::Clamp(ChaosClothConfig->AngularVelocityScale, 0.f, 1.f);
 
-				// Set damping
-				if (ClothSharedSimConfig && ClothSharedSimConfig->bUseDampingOverride)
-				{
-					Evolution->SetDamping(SimDataIndex, ClothSharedSimConfig->Damping);
-				}
-				else
-				{
-					Evolution->SetDamping(SimDataIndex, ChaosClothConfig->DampingCoefficient);
-				}
+				// Set per cloth damping, collision thickness, and friction
+				Evolution->SetDamping(ChaosClothConfig->DampingCoefficient, SimDataIndex);
+				Evolution->SetCollisionThickness(ChaosClothConfig->CollisionThickness, SimDataIndex);
+				Evolution->SetCoefficientOfFriction(ChaosClothConfig->FrictionCoefficient, SimDataIndex);
 
 				// Add Velocity field
 				auto GetVelocity = [this](const TVector<float, 3>&)->TVector<float, 3>
diff --git a/Engine/Plugins/Experimental/ChaosCloth/Source/Chaos/Public/ChaosCloth/ChaosClothConfig.h b/Engine/Plugins/Experimental/ChaosCloth/Source/Chaos/Public/ChaosCloth/ChaosClothConfig.h
--- a/Engine/Plugins/Experimental/ChaosCloth/Source/Chaos/Public/ChaosCloth/ChaosClothConfig.h
+++ b/Engine/Plugins/Experimental/ChaosCloth/Source/Chaos/Public/ChaosCloth/ChaosClothConfig.h
@@ -19,6 +19,9 @@ public:
 	/** Migrate from the legacy FClothConfig structure. */
 	virtual void MigrateFrom(const FClothConfig_Legacy&) override;
 
+	/** Migrate from shared config. */
+	virtual void MigrateFrom(const UClothSharedConfigCommon* ClothSharedConfig) override;
+
 	/** Serialize override used to set the current custom version. */
 	virtual void Serialize(FArchive& Ar) override;
 
@@ -92,14 +95,18 @@ public:
 	UPROPERTY()
 	float ShapeTargetStiffness = 0.f;
 
+	// The radius of cloth points when considering collisions against collider shapes
+	UPROPERTY(EditAnywhere, Category = "Cloth Dynamics", meta = (UIMin = "0", UIMax = "100", ClampMin = "0", ClampMax = "1000"))
+	float CollisionThickness = 1.0f;
+
+	// Friction coefficient for cloth - collider interaction
+	UPROPERTY(EditAnywhere, Category = "Cloth Dynamics", meta = (UIMin = "0", UIMax = "1", ClampMin = "0", ClampMax = "10"))
+	float FrictionCoefficient = 0.2f;
+
 	// The amount of damping applied to the cloth velocities
 	UPROPERTY(EditAnywhere, Category = "Cloth Dynamics", meta = (UIMin = "0", UIMax = "1", ClampMin = "0", ClampMax = "1"))
 	float DampingCoefficient = 0.01f;
 
-	// Friction coefficient for cloth - collider interaction
-	UPROPERTY()
-	float FrictionCoefficient = 0.f;
-
 	// The drag coefficient applying on each particle
 	UPROPERTY(EditAnywhere, Category = "Cloth Dynamics", meta = (UIMin = "0", UIMax = "1", ClampMin = "0", ClampMax = "10"))
 	float DragCoefficient = 0.07f;
@@ -152,7 +159,7 @@ public:
 	virtual ~UChaosClothSharedSimConfig() override;
 
 	/** Migrate from the legacy FClothConfig structure. */
-	virtual void MigrateFrom(const FClothConfig_Legacy&) override;
+	virtual void MigrateFrom(const FClothConfig_Legacy& ClothConfig) override;
 
 	/** Serialize override used to set the current custom version. */
 	virtual void Serialize(FArchive& Ar) override;
@@ -175,16 +182,16 @@ public:
 	float SelfCollisionThickness = 2.0f;
 
 	// The radius of cloth points when considering collisions against collider shapes
-	UPROPERTY(EditAnywhere, Category = Collision, meta = (UIMin = "0", UIMax = "100", ClampMin = "0", ClampMax = "1000"))
-	float CollisionThickness = 1.0f;
+	UPROPERTY()
+	float CollisionThickness_DEPRECATED = 1.0f;
 
 	// Use shared config damping rather than per cloth damping
-	UPROPERTY(EditAnywhere, Category = Simulation, meta = (InlineEditConditionToggle))
-	bool bUseDampingOverride = true;
+	UPROPERTY()
+	bool bUseDampingOverride_DEPRECATED = true;
 
 	// The amount of cloth damping. Override the per cloth damping coefficients.
-	UPROPERTY(EditAnywhere, Category = Simulation, meta = (EditCondition = "bUseDampingOverride", UIMin = "0", UIMax = "1", ClampMin = "0", ClampMax = "1"))
-	float Damping = 0.01f;
+	UPROPERTY()
+	float Damping_DEPRECATED = 0.01f;
 
 	// Use gravity value vs world gravity
 	UPROPERTY(EditAnywhere, Category = Simulation, meta = (InlineEditConditionToggle))
diff --git a/Engine/Source/Runtime/ClothingSystemRuntimeCommon/Private/ClothingAsset.cpp b/Engine/Source/Runtime/ClothingSystemRuntimeCommon/Private/ClothingAsset.cpp
--- a/Engine/Source/Runtime/ClothingSystemRuntimeCommon/Private/ClothingAsset.cpp
+++ b/Engine/Source/Runtime/ClothingSystemRuntimeCommon/Private/ClothingAsset.cpp
@@ -716,6 +716,8 @@ void UClothingAssetCommon::BuildSelfCollisionData()
 void UClothingAssetCommon::PostLoad()
 {
 	Super::PostLoad();
+	// TODO: Remove all UObject PostLoad dependencies.
+	//       Even with these ConditionalPostLoad calls, the UObject PostLoads' order of execution cannot be guaranted.
 	for (UClothLODDataCommon* Lod : ClothLodData)
 	{
 		Lod->ConditionalPostLoad();
@@ -726,15 +728,6 @@ void UClothingAssetCommon::PostLoad()
 
 	if (ClothingCustomVersion < FClothingAssetCustomVersion::MovePropertiesToCommonBaseClasses)
 	{
-		// Remap legacy struct FClothConfig to new config objects
-		for (TPair<FName, UClothConfigBase*>& ClothConfig : ClothConfigs)
-		{
-			if (UClothConfigCommon* const ClothConfigCommon = Cast<UClothConfigCommon>(ClothConfig.Value))
-			{
-				ClothConfigCommon->MigrateFrom(ClothConfig_DEPRECATED);
-			}
-		}
-
 		// Remap legacy struct FClothLODData to class UClothLODDataCommon
 		for (const FClothLODData_Legacy& ClothLODData_Legacy : LodData_DEPRECATED)
 		{
@@ -747,7 +740,7 @@ void UClothingAssetCommon::PostLoad()
 	{
 #if WITH_EDITORONLY_DATA
 		// Convert current parameters to masks
-		for(UClothLODDataCommon* LodPtr : ClothLodData)
+		for (UClothLODDataCommon* LodPtr : ClothLodData)
 		{
 			check(LodPtr);
 			UClothLODDataCommon& Lod = *LodPtr;
@@ -789,7 +782,7 @@ void UClothingAssetCommon::PostLoad()
 
 #if WITH_EDITORONLY_DATA
 	// Fix content imported before we kept vertex colors
-	if(ClothingCustomVersion < FClothingAssetCustomVersion::AddVertexColorsToPhysicalMesh)
+	if (ClothingCustomVersion < FClothingAssetCustomVersion::AddVertexColorsToPhysicalMesh)
 	{
 		for (UClothLODDataCommon* Lod : ClothLodData)
 		{
@@ -806,46 +799,73 @@ void UClothingAssetCommon::PostLoad()
 #endif // WITH_EDITORONLY_DATA
 
 #if WITH_EDITOR
-	if(AnimPhysCustomVersion < FAnimPhysObjectVersion::CacheClothMeshInfluences)
+	if (AnimPhysCustomVersion < FAnimPhysObjectVersion::CacheClothMeshInfluences)
 	{
 		// Rebuild data cache
 		InvalidateCachedData();
 	}
 #endif
 
-	// Migrate simulation dependent config parameters to the new config map
-	if (ClothSimConfig_DEPRECATED)
+	// Add any missing configs for the available cloth factories, and try to migrate them from any existing one
+	AddClothConfigs();
+
+	// Migrate configs
+	bool bMigrateSharedConfigToConfig = true;  // Shared config to config migration can be disabled to avoid overriding the newly migrated values
+
+	if (ClothingCustomVersion < FClothingAssetCustomVersion::MovePropertiesToCommonBaseClasses)
+	{
+		// Remap legacy struct FClothConfig to new config objects
+		for (TPair<FName, UClothConfigBase*>& ClothConfig : ClothConfigs)
+		{
+			if (UClothConfigCommon* const ClothConfigCommon = Cast<UClothConfigCommon>(ClothConfig.Value))
+			{
+				ClothConfigCommon->MigrateFrom(ClothConfig_DEPRECATED);
+			}
+		}
+		bMigrateSharedConfigToConfig = false;
+	}
+	else
 	{
-		// Try a remap to the new config objects through the legacy structure
-		if (const UClothConfigCommon* const ClothSimConfigCommon = Cast<UClothConfigCommon>(ClothSimConfig_DEPRECATED))
+		// Migrate simulation dependent config parameters to the new config map
+		if (ClothSimConfig_DEPRECATED)
 		{
-			FClothConfig_Legacy ClothConfigLegacy;
-			if (ClothSimConfigCommon->MigrateTo(ClothConfigLegacy))
+			// Try a remap to the new config objects through the legacy structure
+			if (const UClothConfigCommon* const ClothSimConfigCommon = Cast<UClothConfigCommon>(ClothSimConfig_DEPRECATED))
 			{
-				for (TPair<FName, UClothConfigBase*>& ClothConfig : ClothConfigs)
+				FClothConfig_Legacy ClothConfigLegacy;
+				if (ClothSimConfigCommon->MigrateTo(ClothConfigLegacy))
 				{
-					if (UClothConfigCommon* const ClothConfigCommon = Cast<UClothConfigCommon>(ClothConfig.Value))
+					for (TPair<FName, UClothConfigBase*>& ClothConfig : ClothConfigs)
 					{
-						ClothConfigCommon->MigrateFrom(ClothConfigLegacy);
+						if (UClothConfigCommon* const ClothConfigCommon = Cast<UClothConfigCommon>(ClothConfig.Value))
+						{
+							ClothConfigCommon->MigrateFrom(ClothConfigLegacy);
+						}
 					}
 				}
 			}
+			// And keep the old config too
+			SetClothConfig(ClothSimConfig_DEPRECATED);
+			ClothSimConfig_DEPRECATED = nullptr;
+			bMigrateSharedConfigToConfig = false;
+		}
+		if (ChaosClothSimConfig_DEPRECATED)
+		{
+			SetClothConfig(ChaosClothSimConfig_DEPRECATED);
+			ChaosClothSimConfig_DEPRECATED = nullptr;
+			bMigrateSharedConfigToConfig = false;
+		}
+		if (ClothSharedSimConfig_DEPRECATED)
+		{
+			SetClothConfig(ClothSharedSimConfig_DEPRECATED);
+			ClothSharedSimConfig_DEPRECATED = nullptr;
+			bMigrateSharedConfigToConfig = false;
 		}
-		// And keep the old config too
-		SetClothConfig(ClothSimConfig_DEPRECATED);
-		ClothSimConfig_DEPRECATED = nullptr;
-	}
-	if (ChaosClothSimConfig_DEPRECATED)
-	{
-		SetClothConfig(ChaosClothSimConfig_DEPRECATED);
-		ChaosClothSimConfig_DEPRECATED = nullptr;
-	}
-	if (ClothSharedSimConfig_DEPRECATED)
-	{
-		SetClothConfig(ClothSharedSimConfig_DEPRECATED);
-		ClothSharedSimConfig_DEPRECATED = nullptr;
 	}
 
+	// Propagate shared configs between cloth assets
+	PropagateSharedConfigs(bMigrateSharedConfigToConfig);
+
 	// After fixing the content, we are ready to call functions that rely on it
 	BuildSelfCollisionData();
 #if WITH_EDITORONLY_DATA
@@ -907,7 +927,7 @@ void UClothingAssetCommon::AddClothConfigs()
 	}
 }
 
-void UClothingAssetCommon::PropagateSharedConfigs()
+void UClothingAssetCommon::PropagateSharedConfigs(bool bMigrateSharedConfigToConfig)
 {
 	// Update this asset's shared config when the asset belongs to a skeletal mesh
 	if (USkeletalMesh* const SkeletalMesh = Cast<USkeletalMesh>(GetOuter()))
@@ -919,7 +939,10 @@ void UClothingAssetCommon::PropagateSharedConfigs()
 
 		for (const UClothingAssetBase* ClothingAssetBase : ClothingAssets)
 		{
-			if (ClothingAssetBase == static_cast<UClothingAssetBase* >(this)) { continue; }
+			if (ClothingAssetBase == static_cast<UClothingAssetBase* >(this))
+			{
+				continue;
+			}
 
 			// Only common assets have shared configs
 			if (const UClothingAssetCommon* const ClothingAsset = Cast<UClothingAssetCommon>(ClothingAssetBase))
@@ -929,29 +952,54 @@ void UClothingAssetCommon::PropagateSharedConfigs()
 				ClothSharedConfigs.Reserve(Max);
 
 				// Iterate through all configs, and find the shared ones
-				for (const TPair<FName, UClothConfigBase*>& ClothConfigItem : ClothingAsset->ClothConfigs)
+				for (const TPair<FName, UClothConfigBase*>& ClothSharedConfigItem : ClothingAsset->ClothConfigs)
 				{
-					if (Cast<UClothSharedConfigCommon>(ClothConfigItem.Value) &&  // Only needs shared configs
-						!ClothSharedConfigs.Find(ClothConfigItem.Key))            // Only needs a single shared config per type
+					if (Cast<UClothSharedConfigCommon>(ClothSharedConfigItem.Value) &&  // Only needs shared configs
+						!ClothSharedConfigs.Find(ClothSharedConfigItem.Key))            // Only needs a single shared config per type
 					{
-						ClothSharedConfigs.Add(ClothConfigItem);
+						ClothSharedConfigs.Add(ClothSharedConfigItem);
 					}
 				}
 			}
 		}
 
 		// Propagate the found shared configs to this asset
-		for (const TPair<FName, UClothConfigBase*>& ClothConfigItem : ClothSharedConfigs)
+		for (const TPair<FName, UClothConfigBase*>& ClothSharedConfigItem : ClothSharedConfigs)
 		{
-			if (UClothConfigBase** const ClothConfigBase = ClothConfigs.Find(ClothConfigItem.Key))
+			// Set share config
+			if (UClothConfigBase** const ClothConfigBase = ClothConfigs.Find(ClothSharedConfigItem.Key))
 			{
 				// Reset this shared config
-				*ClothConfigBase = ClothConfigItem.Value;
+				*ClothConfigBase = ClothSharedConfigItem.Value;
 			}
 			else
 			{
 				// Add new map entry
-				ClothConfigs.Add(ClothConfigItem);
+				ClothConfigs.Add(ClothSharedConfigItem);
+			}
+		}
+
+		// Migrate the common shared configs' deprecated parameters to all per cloth configs
+		if (bMigrateSharedConfigToConfig)
+		{
+			// Iterate through all this asset's shared configs
+			for (const TPair<FName, UClothConfigBase*>& ClothSharedConfigItem : ClothConfigs)
+			{
+				if (const UClothSharedConfigCommon* const ClothSharedConfig = Cast<UClothSharedConfigCommon>(ClothSharedConfigItem.Value))
+				{
+					// Iterate through all this asset's configs, and migrate from the shared ones
+					for (const TPair<FName, UClothConfigBase*>& ClothConfigItem : ClothConfigs)
+					{
+						if (Cast<UClothSharedConfigCommon>(ClothConfigItem.Value))
+						{
+							continue;  // Don't migrate shared configs to another shared configs (or itself)
+						}
+						if (UClothConfigCommon* const ClothConfig = Cast<UClothConfigCommon>(ClothConfigItem.Value))
+						{
+							ClothConfig->MigrateFrom(ClothSharedConfig);
+						}
+					}
+				}
 			}
 		}
 	}
diff --git a/Engine/Source/Runtime/ClothingSystemRuntimeCommon/Public/ClothConfig.h b/Engine/Source/Runtime/ClothingSystemRuntimeCommon/Public/ClothConfig.h
--- a/Engine/Source/Runtime/ClothingSystemRuntimeCommon/Public/ClothConfig.h
+++ b/Engine/Source/Runtime/ClothingSystemRuntimeCommon/Public/ClothConfig.h
@@ -18,6 +18,9 @@ public:
 	/** Migrate from the legacy FClothConfig structure. */
 	virtual void MigrateFrom(const FClothConfig_Legacy&) {}
 
+	/** Migrate from shared configs. */
+	virtual void MigrateFrom(const class UClothSharedConfigCommon*) {}
+
 	/**
 	 * Migrate to the legacy FClothConfig structure.
 	 * Useful for converting configs that are compatible with this legacy structure.
diff --git a/Engine/Source/Runtime/ClothingSystemRuntimeCommon/Public/ClothingAsset.h b/Engine/Source/Runtime/ClothingSystemRuntimeCommon/Public/ClothingAsset.h
--- a/Engine/Source/Runtime/ClothingSystemRuntimeCommon/Public/ClothingAsset.h
+++ b/Engine/Source/Runtime/ClothingSystemRuntimeCommon/Public/ClothingAsset.h
@@ -267,8 +267,9 @@ private:
 	void AddClothConfigs();
 
 	// Propagate the shared simulation configs between assets.
+	// Also migrate all deprecated shared parameters which have been moved to the per cloth configs if required.
 	// Called after a cloth asset is created or loaded.
-	void PropagateSharedConfigs();
+	void PropagateSharedConfigs(bool bMigrateSharedConfigToConfig=false);
 
 #if WITH_EDITOR
 	// Helper functions used in PostPropertyChangeCb
diff --git a/Engine/Source/Runtime/Engine/Private/SkeletalMesh.cpp b/Engine/Source/Runtime/Engine/Private/SkeletalMesh.cpp
--- a/Engine/Source/Runtime/Engine/Private/SkeletalMesh.cpp
+++ b/Engine/Source/Runtime/Engine/Private/SkeletalMesh.cpp
@@ -2324,14 +2324,16 @@ void USkeletalMesh::PostLoad()
 	LLM_SCOPE(ELLMTag::SkeletalMesh);
 	Super::PostLoad();
 
-	// Consolidate the shared cloth configs once all cloth assets are loaded
+	// Make sure the cloth assets have finished loading
+	// TODO: Remove all UObject PostLoad dependencies.
+	//       Even with these ConditionalPostLoad calls, the UObject PostLoads' order of execution cannot be guaranted.
+	//       E.g. in some instance it has been found that the SkeletalMesh EndLoad can trigger a ConditionalPostLoad
+	//       on the cloth assets even before reaching this point.
+	//       In these occurences, the cloth asset's RF_NeedsPostLoad flag is already cleared despite its PostLoad still
+	//       being un-executed, making the following block code ineffective.
 	for (UClothingAssetBase* MeshClothingAsset : MeshClothingAssets)
 	{
-		MeshClothingAsset->ConditionalPostLoad();  // Make sure the cloth asset has finished loading
-	}
-	for (UClothingAssetBase* MeshClothingAsset : MeshClothingAssets)  // PostUpdateAllAssets will also iterate through all clothing assets so this cannot be merged with the loop above
-	{
-		MeshClothingAsset->PostUpdateAllAssets();
+		MeshClothingAsset->ConditionalPostLoad();
 	}
 
 #if WITH_EDITOR
diff --git a/Engine/Source/Runtime/Experimental/Chaos/Private/Chaos/PBDEvolution.cpp b/Engine/Source/Runtime/Experimental/Chaos/Private/Chaos/PBDEvolution.cpp
--- a/Engine/Source/Runtime/Experimental/Chaos/Private/Chaos/PBDEvolution.cpp
+++ b/Engine/Source/Runtime/Experimental/Chaos/Private/Chaos/PBDEvolution.cpp
@@ -45,6 +45,8 @@ TPBDEvolution<T, d>::TPBDEvolution(TPBDParticles<T, d>&& InParticles, TKinematic
 	MCollisionParticles.AddArray(&MCollisionParticleGroupIds);
 	MParticles.AddArray(&MParticleGroupIds);
 	MPerGroupDamping.Add(Damping);
+	MPerGroupCollisionThickness.Add(CollisionThickness);
+	MPerGroupCoefficientOfFriction.Add(CoefficientOfFriction);
 
 	SetParticleUpdateFunction(
 		[PBDUpdateRule = 
@@ -76,9 +78,13 @@ uint32 TPBDEvolution<T, d>::AddParticles(uint32 Num, uint32 GroupId)
 	if (GroupId >= GroupNum)
 	{
 		MPerGroupDamping.SetNum((int32)GroupId + 1);
+		MPerGroupCollisionThickness.SetNum((int32)GroupId + 1);
+		MPerGroupCoefficientOfFriction.SetNum((int32)GroupId + 1);
 		for (uint32 i = GroupNum; i <= GroupId; ++i)
 		{
 			MPerGroupDamping[i] = MDamping;
+			MPerGroupCollisionThickness[i] = MCollisionThickness;
+			MPerGroupCoefficientOfFriction[i] = MCoefficientOfFriction;
 		}
 	}
 	return Offset;
@@ -107,7 +113,7 @@ void TPBDEvolution<T, d>::AdvanceOneTimeStep(const T Dt)
 	TPerParticleEulerStepVelocity<T, d> EulerStepVelocityRule;
 	TPerGroupDampVelocity<T, d> DampVelocityRule(MParticleGroupIds, MPerGroupDamping);
 	TPerParticlePBDEulerStep<T, d> EulerStepRule;
-	TPerParticlePBDCollisionConstraint<T, d, EGeometryParticlesSimType::Other> CollisionRule(MCollisionParticles, MCollided, MParticleGroupIds, MCollisionParticleGroupIds, MCollisionThickness, MCoefficientOfFriction);
+	TPerParticlePBDCollisionConstraint<T, d, EGeometryParticlesSimType::Other> CollisionRule(MCollisionParticles, MCollided, MParticleGroupIds, MCollisionParticleGroupIds, MPerGroupCollisionThickness, MPerGroupCoefficientOfFriction);
 
 	{
 		SCOPE_CYCLE_COUNTER(STAT_ChaosPBDVelocityDampUpdateState);
diff --git a/Engine/Source/Runtime/Experimental/Chaos/Public/Chaos/PBDEvolution.h b/Engine/Source/Runtime/Experimental/Chaos/Public/Chaos/PBDEvolution.h
--- a/Engine/Source/Runtime/Experimental/Chaos/Public/Chaos/PBDEvolution.h
+++ b/Engine/Source/Runtime/Experimental/Chaos/Public/Chaos/PBDEvolution.h
@@ -56,15 +56,14 @@ class CHAOS_API TPBDEvolution
 	T GetSelfCollisionThickness() const { return MSelfCollisionThickness; }
 	void SetSelfCollisionThickness(const T SelfCollisionThickness) { MSelfCollisionThickness = SelfCollisionThickness; }
 
-	T GetCollisionThickness() const { return MCollisionThickness; }
-	void SetCollisionThickness(const T CollisionThickness) { MCollisionThickness = CollisionThickness; }
+	T GetCollisionThickness(const uint32 GroupId = 0) const { return MPerGroupCollisionThickness[GroupId]; }
+	void SetCollisionThickness(const T CollisionThickness, const uint32 GroupId = 0) { MPerGroupCollisionThickness[GroupId] = CollisionThickness; }
 
-	T GetCoefficientOfFriction() const { return MCoefficientOfFriction; }
-	void SetCoefficientOfFriction(const T CoefficientOfFriction) { MCoefficientOfFriction = CoefficientOfFriction; }
+	T GetCoefficientOfFriction(const uint32 GroupId = 0) const { return MPerGroupCoefficientOfFriction[GroupId]; }
+	void SetCoefficientOfFriction(const T CoefficientOfFriction, const uint32 GroupId = 0) { MPerGroupCoefficientOfFriction[GroupId] = CoefficientOfFriction; }
 
-	T GetDamping() const { return MDamping; }
-	void SetDamping(const T Damping) { MDamping = Damping; }
-	void SetDamping(const uint32 GroupId, const T Damping) { MPerGroupDamping[GroupId] = Damping; }
+	T GetDamping(const uint32 GroupId = 0) const { return MPerGroupDamping[GroupId]; }
+	void SetDamping(const T Damping, const uint32 GroupId = 0) { MPerGroupDamping[GroupId] = Damping; }
 
 	T GetTime() const { return MTime; }
 
@@ -81,6 +80,8 @@ class CHAOS_API TPBDEvolution
 	TArrayCollectionArray<uint32> MCollisionParticleGroupIds;  // Used for per group parameters for collision particles
 	TArrayCollectionArray<uint32> MParticleGroupIds;  // Used for per group parameters for particles
 	TArray<T> MPerGroupDamping;
+	TArray<T> MPerGroupCollisionThickness;
+	TArray<T> MPerGroupCoefficientOfFriction;
 	int32 MNumIterations;
 	T MCollisionThickness;
 	T MSelfCollisionThickness;
diff --git a/Engine/Source/Runtime/Experimental/Chaos/Public/Chaos/PerParticlePBDCollisionConstraint.h b/Engine/Source/Runtime/Experimental/Chaos/Public/Chaos/PerParticlePBDCollisionConstraint.h
--- a/Engine/Source/Runtime/Experimental/Chaos/Public/Chaos/PerParticlePBDCollisionConstraint.h
+++ b/Engine/Source/Runtime/Experimental/Chaos/Public/Chaos/PerParticlePBDCollisionConstraint.h
@@ -18,8 +18,8 @@ class TPerParticlePBDCollisionConstraint : public TPerParticleRule<T, d>
 	};
 
   public:
-	TPerParticlePBDCollisionConstraint(const TKinematicGeometryParticlesImp<T, d, SimType>& InParticles, TArrayCollectionArray<bool>& Collided, TArrayCollectionArray<uint32>& DynamicGroupIds, TArrayCollectionArray<uint32>& KinematicGroupIds, const T Thickness = (T)0, const T Friction = (T)0)
-	    : bFastPositionBasedFriction(true), MParticles(InParticles), MCollided(Collided), MDynamicGroupIds(DynamicGroupIds), MKinematicGroupIds(KinematicGroupIds), MThickness(Thickness), MFriction(Friction) {}
+	TPerParticlePBDCollisionConstraint(const TKinematicGeometryParticlesImp<T, d, SimType>& InParticles, TArrayCollectionArray<bool>& Collided, TArrayCollectionArray<uint32>& DynamicGroupIds, TArrayCollectionArray<uint32>& KinematicGroupIds, const TArray<T>& PerGroupThickness, const TArray<T>& PerGroupFriction)
+	    : bFastPositionBasedFriction(true), MParticles(InParticles), MCollided(Collided), MDynamicGroupIds(DynamicGroupIds), MKinematicGroupIds(KinematicGroupIds), MPerGroupThickness(PerGroupThickness), MPerGroupFriction(PerGroupFriction) {}
 	virtual ~TPerParticlePBDCollisionConstraint() {}
 
 	inline void Apply(TPBDParticles<T, d>& InParticles, const T Dt, const int32 Index) const override //-V762
@@ -28,25 +28,30 @@ class TPerParticlePBDCollisionConstraint : public TPerParticleRule<T, d>
 			return;
 		for (uint32 i = 0; i < MParticles.Size(); ++i)
 		{
-			if (MDynamicGroupIds[Index] != MKinematicGroupIds[i]) { continue; }  // Bail out if the collision groups doesn't match the particle group id
+			if (MDynamicGroupIds[Index] != MKinematicGroupIds[i])
+			{
+				continue; // Bail out if the collision groups doesn't match the particle group id
+			}
 			TVector<T, d> Normal;
 			TRigidTransform<T, d> Frame(MParticles.X(i), MParticles.R(i));
-			T Phi = MParticles.Geometry(i)->PhiWithNormal(Frame.InverseTransformPosition(InParticles.P(Index)), Normal);
-			if (Phi < MThickness)
+			const T Phi = MParticles.Geometry(i)->PhiWithNormal(Frame.InverseTransformPosition(InParticles.P(Index)), Normal);
+			const T Thickness = MPerGroupThickness[MDynamicGroupIds[Index]];
+			if (Phi < Thickness)
 			{
 				const TVector<T, d> NormalWorld = Frame.TransformVector(Normal);
-				InParticles.P(Index) += (-Phi + MThickness) * NormalWorld;
-				if (MFriction > 0)
+				InParticles.P(Index) += (-Phi + Thickness) * NormalWorld;
+				const T CoefficientOfFriction = MPerGroupFriction[MDynamicGroupIds[Index]];
+				if (CoefficientOfFriction > 0)
 				{
 					const T MaximumFrictionCorrectionPerStep = 1.0f; // 1cm absolute maximum correction the friction force can provide
 					if (bFastPositionBasedFriction)
 					{
-						const T Penetration = (-Phi + MThickness); // This is related to the Normal impulse 
+						const T Penetration = (-Phi + Thickness); // This is related to the Normal impulse
 						TVector<T, d> VectorToPoint = InParticles.P(Index) - MParticles.X(i);
 						const TVector<T, d> RelativeDisplacement = (InParticles.P(Index) - InParticles.X(Index)) - (MParticles.V(i) + TVector<T, d>::CrossProduct(MParticles.W(i), VectorToPoint)) * Dt; // This corresponds to the tangential velocity multiplied by dt (friction will drive this to zero if it is high enough)
 						const TVector<T, d> RelativeDisplacementTangent = RelativeDisplacement - TVector<T, d>::DotProduct(RelativeDisplacement, NormalWorld) * NormalWorld; // Project displacement into the tangential plane
 						const T RelativeDisplacementTangentLength = RelativeDisplacementTangent.Size();
-						T PositionCorrection = FMath::Clamp<T>(Penetration * MFriction, 0, RelativeDisplacementTangentLength);
+						T PositionCorrection = FMath::Clamp<T>(Penetration * CoefficientOfFriction, 0, RelativeDisplacementTangentLength);
 						if (PositionCorrection > MaximumFrictionCorrectionPerStep)
 						{
 							PositionCorrection = MaximumFrictionCorrectionPerStep;
@@ -75,31 +80,33 @@ class TPerParticlePBDCollisionConstraint : public TPerParticleRule<T, d>
 	{
 		if (!bFastPositionBasedFriction)
 		{
-			check(MFriction > 0);
 			if (!MVelocityConstraints.Contains(Index))
 			{
 				return;
 			}
-			T VN = TVector<T, d>::DotProduct(InParticles.V(Index), MVelocityConstraints[Index].Normal);
-			T VNBody = TVector<T, d>::DotProduct(MVelocityConstraints[Index].Velocity, MVelocityConstraints[Index].Normal);
-			TVector<T, d> VTBody = MVelocityConstraints[Index].Velocity - VNBody * MVelocityConstraints[Index].Normal;
-			TVector<T, d> VTRelative = InParticles.V(Index) - VN * MVelocityConstraints[Index].Normal - VTBody;
-			T VTRelativeSize = VTRelative.Size();
-			T VNMax = FGenericPlatformMath::Max(VN, VNBody);
-			T VNDelta = VNMax - VN;
-			T Friction = MFriction * VNDelta < VTRelativeSize ? MFriction * VNDelta / VTRelativeSize : 1;
+			const T VN = TVector<T, d>::DotProduct(InParticles.V(Index), MVelocityConstraints[Index].Normal);
+			const T VNBody = TVector<T, d>::DotProduct(MVelocityConstraints[Index].Velocity, MVelocityConstraints[Index].Normal);
+			const TVector<T, d> VTBody = MVelocityConstraints[Index].Velocity - VNBody * MVelocityConstraints[Index].Normal;
+			const TVector<T, d> VTRelative = InParticles.V(Index) - VN * MVelocityConstraints[Index].Normal - VTBody;
+			const T VTRelativeSize = VTRelative.Size();
+			const T VNMax = FGenericPlatformMath::Max(VN, VNBody);
+			const T VNDelta = VNMax - VN;
+			const T CoefficientOfFriction = MPerGroupFriction[MDynamicGroupIds[Index]];
+			check(CoefficientOfFriction > 0);
+			const T Friction = CoefficientOfFriction * VNDelta < VTRelativeSize ? CoefficientOfFriction * VNDelta / VTRelativeSize : 1;
 			InParticles.V(Index) = VNMax * MVelocityConstraints[Index].Normal + VTBody + VTRelative * (1 - Friction);
 		}
 	}
 
   private:
-	  bool bFastPositionBasedFriction;
+	bool bFastPositionBasedFriction;
 	// TODO(mlentine): Need a bb hierarchy
 	const TKinematicGeometryParticlesImp<T, d, SimType>& MParticles;
 	TArrayCollectionArray<bool>& MCollided;
 	TArrayCollectionArray<uint32>& MDynamicGroupIds;
 	TArrayCollectionArray<uint32>& MKinematicGroupIds;
 	mutable TMap<int32, VelocityConstraint> MVelocityConstraints;
-	const T MThickness, MFriction;
+	const TArray<T>& MPerGroupThickness;
+	const TArray<T>& MPerGroupFriction;
 };
 }
