Cloth - Fix old Apex cloth import that was missing the correct config.
- Make sure the current configs have finished loading before adding any mssing configs in UClothingAssetCommon::PostLoad().

diff --git a/Engine/Source/Editor/ClothingSystemEditor/Private/ClothingAssetFactory.cpp b/Engine/Source/Editor/ClothingSystemEditor/Private/ClothingAssetFactory.cpp
--- a/Engine/Source/Editor/ClothingSystemEditor/Private/ClothingAssetFactory.cpp
+++ b/Engine/Source/Editor/ClothingSystemEditor/Private/ClothingAssetFactory.cpp
@@ -388,6 +388,7 @@ UClothingAssetBase* UClothingAssetFactory::CreateFromApexAsset(nvidia::apex::Clo
 #if WITH_APEX_CLOTHING
 	UClothingAssetCommon* NewClothingAsset = NewObject<UClothingAssetCommon>(TargetMesh, InName);
 	NewClothingAsset->SetFlags(RF_Transactional);
+	NewClothingAsset->PostUpdateAllAssets();  // Ensure that the required configs, including shared configs, are created
 
 	const NvParameterized::Interface* AssetParams = InApexAsset->getAssetNvParameterized();
 	NvParameterized::Handle GraphicalLodArrayHandle(*AssetParams, "graphicalLods");
diff --git a/Engine/Source/Runtime/ClothingSystemRuntimeCommon/Private/ClothingAsset.cpp b/Engine/Source/Runtime/ClothingSystemRuntimeCommon/Private/ClothingAsset.cpp
--- a/Engine/Source/Runtime/ClothingSystemRuntimeCommon/Private/ClothingAsset.cpp
+++ b/Engine/Source/Runtime/ClothingSystemRuntimeCommon/Private/ClothingAsset.cpp
@@ -717,7 +717,7 @@ void UClothingAssetCommon::PostLoad()
 {
 	Super::PostLoad();
 	// TODO: Remove all UObject PostLoad dependencies.
-	//       Even with these ConditionalPostLoad calls, the UObject PostLoads' order of execution cannot be guaranted.
+	//       Even with these ConditionalPostLoad calls, the UObject PostLoads' order of execution cannot be guaranteed.
 	for (UClothLODDataCommon* Lod : ClothLodData)
 	{
 		Lod->ConditionalPostLoad();
@@ -807,6 +807,15 @@ void UClothingAssetCommon::PostLoad()
 #endif
 
 	// Add any missing configs for the available cloth factories, and try to migrate them from any existing one
+	// TODO: Remove all UObject PostLoad dependencies.
+	//       Even with these ConditionalPostLoad calls, the UObject PostLoads' order of execution cannot be guaranteed.
+	for (TPair<FName, UClothConfigBase*>& ClothConfig : ClothConfigs)
+	{
+		if (ClothConfig.Value)
+		{
+			ClothConfig.Value->ConditionalPostLoad();  // PostLoad configs before adding new ones
+		}
+	}
 	AddClothConfigs();
 
 	// Migrate configs
