Chaos Cloth - Rollback UClothLODDataCommon to FClothLODData to prevent crashes in SkeletalMesh PostLoad.

diff --git a/Engine/Config/BaseEngine.ini b/Engine/Config/BaseEngine.ini
--- a/Engine/Config/BaseEngine.ini
+++ b/Engine/Config/BaseEngine.ini
@@ -707,7 +707,7 @@ bUseStreamingPause=false
 
 ; Redirect legacy ClothingSystemRuntime components to ClothingSystemRuntimeCommon
 +ClassRedirects=(OldName="ClothingAsset",NewName="/Script/ClothingSystemRuntimeCommon.ClothingAssetCommon")
-+StructRedirects=(OldName="ClothLODData",NewName="/Script/ClothingSystemRuntimeCommon.ClothLODData_Legacy")
++StructRedirects=(OldName="ClothLODData",NewName="/Script/ClothingSystemRuntimeCommon.ClothLODDataCommon")
 +StructRedirects=(OldName="ClothConfig",NewName="/Script/ClothingSystemRuntimeCommon.ClothConfig_Legacy")
 +StructRedirects=(OldName="ClothParameterMask_PhysMesh",NewName="/Script/ClothingSystemRuntimeCommon.ClothParameterMask_Legacy")
 +StructRedirects=(OldName="ClothConstraintSetup",NewName="/Script/ClothingSystemRuntimeCommon.ClothConstraintSetup_Legacy")
@@ -715,9 +715,12 @@ bUseStreamingPause=false
 
 ; Redirect common ClothingSystemRuntimeNv components to ClothingSystemRuntimeCommon
 +ClassRedirects=(OldName="/Script/ClothingSystemRuntimeNv.ClothingAssetNv",NewName="/Script/ClothingSystemRuntimeCommon.ClothingAssetCommon")
-+ClassRedirects=(OldName="/Script/ClothingSystemRuntimeNv.ClothLODDataNv",NewName="/Script/ClothingSystemRuntimeCommon.ClothLODDataCommon")
++ClassRedirects=(OldName="/Script/ClothingSystemRuntimeNv.ClothLODDataNv",NewName="/Script/ClothingSystemRuntimeCommon.ClothLODDataCommon_Legacy")
 +EnumRedirects=(OldName="MaskTarget_PhysMesh",NewName="/Script/ClothingSystemRuntimeCommon.EWeightMapTargetCommon")
 
+; Deprecate the ClothLODDataCommon UObject class that was causing all sort of PostLoad issues
++ClassRedirects=(OldName="/Script/ClothingSystemRuntimeCommon.ClothLODDataCommon",NewName="/Script/ClothingSystemRuntimeCommon.ClothLODDataCommon_Legacy")
+
 ; Deprecate the ClothPhysicalMeshDataNv class, but keep it around as a legacy class for clothing asset backward compatibility
 +ClassRedirects=(OldName="/Script/ClothingSystemRuntimeInterface.ClothPhysicalMeshDataBase",NewName="/Script/ClothingSystemRuntimeInterface.ClothPhysicalMeshDataBase_Legacy")
 +ClassRedirects=(OldName="/Script/ClothingSystemRuntimeNv.ClothPhysicalMeshDataNv",NewName="/Script/ClothingSystemRuntimeNv.ClothPhysicalMeshDataNv_Legacy")
diff --git a/Engine/Plugins/Experimental/ChaosCloth/Source/Chaos/Private/ChaosCloth/ChaosClothingSimulation.cpp b/Engine/Plugins/Experimental/ChaosCloth/Source/Chaos/Private/ChaosCloth/ChaosClothingSimulation.cpp
--- a/Engine/Plugins/Experimental/ChaosCloth/Source/Chaos/Private/ChaosCloth/ChaosClothingSimulation.cpp
+++ b/Engine/Plugins/Experimental/ChaosCloth/Source/Chaos/Private/ChaosCloth/ChaosClothingSimulation.cpp
@@ -238,8 +238,8 @@ void ClothingSimulation::CreateActor(USkeletalMeshComponent* InOwnerComponent, U
 	UE_CLOG(Asset->GetNumLods() != 1,
 		LogChaosCloth, Warning, TEXT("More than one LOD with the current cloth asset %s in sim slot %d. Only LOD 0 is supported with Chaos Cloth for now."),
 		InOwnerComponent->GetOwner() ? *InOwnerComponent->GetOwner()->GetName() : TEXT("None"), InSimDataIndex);
-	const UClothLODDataCommon* const AssetLodData = Asset->ClothLodData[0];
-	const FClothPhysicalMeshData& PhysMesh = AssetLodData->ClothPhysicalMeshData;
+	const FClothLODDataCommon& AssetLodData = Asset->LodData[0];
+	const FClothPhysicalMeshData& PhysMesh = AssetLodData.PhysicalMeshData;
 
 	// Add particles
 	TPBDParticles<float, 3>& Particles = Evolution->Particles();
@@ -326,7 +326,7 @@ void ClothingSimulation::CreateActor(USkeletalMeshComponent* InOwnerComponent, U
 	}
 
 	// Warn about legacy apex collisions
-	const FClothCollisionData& LodCollData = AssetLodData->CollisionData;
+	const FClothCollisionData& LodCollData = AssetLodData.CollisionData;
 	UE_CLOG(LodCollData.Spheres.Num() > 0 || LodCollData.SphereConnections.Num() > 0 || LodCollData.Convexes.Num() > 0,
 		LogChaosCloth, Warning, TEXT(
 			"Actor '%s' component '%s' has %d sphere, %d capsule, and %d "
@@ -1040,14 +1040,12 @@ void ClothingSimulation::ExtractPhysicsAssetCollisions(const UClothingAssetCommo
 
 void ClothingSimulation::ExtractLegacyAssetCollisions(const UClothingAssetCommon* Asset, int32 InSimDataIndex)
 {
-	if (const UClothLODDataCommon* const AssetLodData = Asset->ClothLodData[0])
+	const FClothLODDataCommon& AssetLodData = Asset->LodData[0];
+	const FClothCollisionData& LodCollData = AssetLodData.CollisionData;
+	if (LodCollData.Spheres.Num() || LodCollData.SphereConnections.Num() || LodCollData.Convexes.Num())
 	{
-		const FClothCollisionData& LodCollData = AssetLodData->CollisionData;
-		if (LodCollData.Spheres.Num() || LodCollData.SphereConnections.Num() || LodCollData.Convexes.Num())
-		{
-			UE_LOG(LogChaosCloth, VeryVerbose, TEXT("Adding legacy cloth asset collisions..."));
-			AddCollisions(LodCollData, Asset->UsedBoneIndices, InSimDataIndex);
-		}
+		UE_LOG(LogChaosCloth, VeryVerbose, TEXT("Adding legacy cloth asset collisions..."));
+		AddCollisions(LodCollData, Asset->UsedBoneIndices, InSimDataIndex);
 	}
 }
 
@@ -1350,7 +1348,7 @@ void ClothingSimulation::Simulate(IClothingSimulationContext* InContext)
 
 			ClothingMeshUtils::SkinPhysicsMesh<true, false>(
 				Asset->UsedBoneIndices,
-				Asset->ClothLodData[0]->ClothPhysicalMeshData,
+				Asset->LodData[0].PhysicalMeshData,
 				ComponentToLocalSpace,
 				Context->RefToLocals.GetData(),
 				Context->RefToLocals.Num(),
@@ -1668,9 +1666,10 @@ void ClothingSimulation::GetCollisions(FClothCollisionData& OutCollisions, bool
 	// Add internal asset collisions
 	for (const UClothingAssetCommon* Asset : Assets)
 	{
-		if (const UClothLODDataCommon* const ClothLodData = !Asset ? nullptr : Asset->ClothLodData[0])
+		if (Asset)
 		{
-			OutCollisions.Append(ClothLodData->CollisionData);
+			const FClothLODDataCommon& ClothLodData = Asset->LodData[0];
+			OutCollisions.Append(ClothLodData.CollisionData);
 		}
 	}
 
@@ -1698,7 +1697,7 @@ void ClothingSimulation::RefreshClothConfig()
 			if (const UChaosClothConfig* const ChaosClothConfig = Asset->GetClothConfig<UChaosClothConfig>())
 			{
 				check(Asset->GetNumLods() > 0);
-				const FClothPhysicalMeshData& PhysMesh = Asset->ClothLodData[0]->ClothPhysicalMeshData;
+				const FClothPhysicalMeshData& PhysMesh = Asset->LodData[0].PhysicalMeshData;
 
 				ResetParticles(SimDataIndex);
 
@@ -2246,9 +2245,8 @@ void ClothingSimulation::DebugDrawBackstops(USkeletalMeshComponent* /*OwnerCompo
 		if (!Asset) { continue; }
 
 		// Get Backstop Distances
-		const UClothLODDataCommon* const AssetLodData = Asset->ClothLodData[0];
-		check(AssetLodData);
-		const FClothPhysicalMeshData& PhysMesh = AssetLodData->ClothPhysicalMeshData;
+		const FClothLODDataCommon& AssetLodData = Asset->LodData[0];
+		const FClothPhysicalMeshData& PhysMesh = AssetLodData.PhysicalMeshData;
 		const FPointWeightMap& BackstopDistances = PhysMesh.GetWeightMap(EChaosWeightMapTarget::BackstopDistance);
 		const FPointWeightMap& BackstopRadiuses = PhysMesh.GetWeightMap(EChaosWeightMapTarget::BackstopRadius);
 		if (BackstopDistances.Num() == 0 || BackstopRadiuses.Num() == 0)
@@ -2294,9 +2292,8 @@ void ClothingSimulation::DebugDrawMaxDistances(USkeletalMeshComponent* /*OwnerCo
 		if (!Asset) { continue; }
 
 		// Get Maximum Distances
-		const UClothLODDataCommon* const AssetLodData = Asset->ClothLodData[0];
-		check(AssetLodData);
-		const FClothPhysicalMeshData& PhysMesh = AssetLodData->ClothPhysicalMeshData;
+		const FClothLODDataCommon& AssetLodData = Asset->LodData[0];
+		const FClothPhysicalMeshData& PhysMesh = AssetLodData.PhysicalMeshData;
 		const FPointWeightMap& MaxDistances = PhysMesh.GetWeightMap(EChaosWeightMapTarget::MaxDistance);
 		if (MaxDistances.Num() == 0)
 		{
@@ -2331,9 +2328,8 @@ void ClothingSimulation::DebugDrawAnimDrive(USkeletalMeshComponent* /*OwnerCompo
 		if (!Asset) { continue; }
 
 		// Get anim drive multiplier
-		const UClothLODDataCommon* const AssetLodData = Asset->ClothLodData[0];
-		check(AssetLodData);
-		const FClothPhysicalMeshData& PhysMesh = AssetLodData->ClothPhysicalMeshData;
+		const FClothLODDataCommon& AssetLodData = Asset->LodData[0];
+		const FClothPhysicalMeshData& PhysMesh = AssetLodData.PhysicalMeshData;
 		const FPointWeightMap& AnimDriveMultipliers = PhysMesh.GetWeightMap(EChaosWeightMapTarget::AnimDriveMultiplier);
 		if (AnimDriveMultipliers.Num() == 0)
 		{
diff --git a/Engine/Source/Editor/ClothPainter/Private/ClothMeshAdapter.cpp b/Engine/Source/Editor/ClothPainter/Private/ClothMeshAdapter.cpp
--- a/Engine/Source/Editor/ClothPainter/Private/ClothMeshAdapter.cpp
+++ b/Engine/Source/Editor/ClothPainter/Private/ClothMeshAdapter.cpp
@@ -56,7 +56,7 @@ bool FClothMeshPaintAdapter::Initialize()
 
 	UClothingAssetCommon* ClothingAsset = Cast<UClothingAssetCommon>(SelectedAsset);
 
-	bHaveAsset = ClothingAsset && ClothingAsset->ClothLodData.IsValidIndex(PaintingClothLODIndex);
+	bHaveAsset = ClothingAsset && ClothingAsset->LodData.IsValidIndex(PaintingClothLODIndex);
 	bBaseInit = FBaseMeshPaintGeometryAdapter::Initialize();
 
 	return bHaveAsset && bBaseInit;
@@ -182,9 +182,9 @@ void FClothMeshPaintAdapter::SetSelectedClothingAsset(const FGuid& InAssetGuid,
 			{
 				if(ConcreteAsset->IsValidLod(InAssetLod))
 				{
-					UClothLODDataCommon* LodData = ConcreteAsset->ClothLodData[InAssetLod];
+					FClothLODDataCommon& LodData = ConcreteAsset->LodData[InAssetLod];
 
-					if(LodData->ParameterMasks.IsValidIndex(InMaskIndex))
+					if(LodData.PointWeightMaps.IsValidIndex(InMaskIndex))
 					{
 						PaintingClothLODIndex = InAssetLod;
 						PaintingClothMaskIndex = InMaskIndex;
@@ -233,8 +233,8 @@ bool FClothMeshPaintAdapter::InitializeVertexData()
 			if(DebugComponent->SkinnedSelectedClothingPositions.Num() > 0)
 			{
 				UClothingAssetCommon* ConcreteAsset = CastChecked<UClothingAssetCommon>(SelectedAsset);
-				const UClothLODDataCommon* LODData = ConcreteAsset->ClothLodData[PaintingClothLODIndex];
-				const FClothPhysicalMeshData& MeshData = LODData->ClothPhysicalMeshData;
+				const FClothLODDataCommon& LODData = ConcreteAsset->LodData[PaintingClothLODIndex];
+				const FClothPhysicalMeshData& MeshData = LODData.PhysicalMeshData;
 
 				MeshVertices.Append(DebugComponent->SkinnedSelectedClothingPositions);
 				MeshIndices.Append(MeshData.Indices);
@@ -296,7 +296,7 @@ FPointWeightMap* FClothMeshPaintAdapter::GetCurrentMask() const
 	{
 		UClothingAssetCommon* ConcreteAsset = CastChecked<UClothingAssetCommon>(SelectedAsset);
 
-		return &ConcreteAsset->ClothLodData[PaintingClothLODIndex]->ParameterMasks[PaintingClothMaskIndex];
+		return &ConcreteAsset->LodData[PaintingClothLODIndex].PointWeightMaps[PaintingClothMaskIndex];
 	}
 
 	return nullptr;
@@ -309,8 +309,8 @@ bool FClothMeshPaintAdapter::HasValidSelection() const
 		UClothingAssetCommon* ConcreteAsset = Cast<UClothingAssetCommon>(SelectedAsset);
 
 		// Only valid selection if we have a valid asset, asset LOD and a mask
-		if(ConcreteAsset->ClothLodData.IsValidIndex(PaintingClothLODIndex) &&
-			ConcreteAsset->ClothLodData[PaintingClothLODIndex]->ParameterMasks.IsValidIndex(PaintingClothMaskIndex))
+		if(ConcreteAsset->LodData.IsValidIndex(PaintingClothLODIndex) &&
+			ConcreteAsset->LodData[PaintingClothLODIndex].PointWeightMaps.IsValidIndex(PaintingClothMaskIndex))
 		{
 			return true;
 		}
diff --git a/Engine/Source/Editor/ClothPainter/Private/ClothPainter.cpp b/Engine/Source/Editor/ClothPainter/Private/ClothPainter.cpp
--- a/Engine/Source/Editor/ClothPainter/Private/ClothPainter.cpp
+++ b/Engine/Source/Editor/ClothPainter/Private/ClothPainter.cpp
@@ -564,8 +564,8 @@ void FClothPainter::OnAssetSelectionChanged(UClothingAssetCommon* InNewSelectedA
 	{
 		// Validate the incoming parameters, to make sure we only set a selection if we're going
 		// to get a valid paintable surface
-		if(InNewSelectedAsset->ClothLodData.IsValidIndex(InAssetLod) &&
-			 InNewSelectedAsset->ClothLodData[InAssetLod]->ParameterMasks.IsValidIndex(InMaskIndex))
+		if(InNewSelectedAsset->LodData.IsValidIndex(InAssetLod) &&
+			 InNewSelectedAsset->LodData[InAssetLod].PointWeightMaps.IsValidIndex(InMaskIndex))
 		{
 			const FGuid NewGuid = InNewSelectedAsset->GetAssetGuid();
 			SkeletalMeshComponent->SetMeshSectionVisibilityForCloth(SkeletalMeshComponent->SelectedClothingGuidForPainting, true);
diff --git a/Engine/Source/Editor/ClothPainter/Private/SClothAssetSelector.cpp b/Engine/Source/Editor/ClothPainter/Private/SClothAssetSelector.cpp
--- a/Engine/Source/Editor/ClothPainter/Private/SClothAssetSelector.cpp
+++ b/Engine/Source/Editor/ClothPainter/Private/SClothAssetSelector.cpp
@@ -39,10 +39,10 @@ FPointWeightMap* FClothingMaskListItem::GetMask()
 	{
 		if(Asset->IsValidLod(LodIndex))
 		{
-			UClothLODDataCommon* LodData = Asset->ClothLodData[LodIndex];
-			if(LodData->ParameterMasks.IsValidIndex(MaskIndex))
+			FClothLODDataCommon& LodData = Asset->LodData[LodIndex];
+			if(LodData.PointWeightMaps.IsValidIndex(MaskIndex))
 			{
-				return &LodData->ParameterMasks[MaskIndex];
+				return &LodData.PointWeightMaps[MaskIndex];
 			}
 		}
 	}
@@ -52,13 +52,13 @@ FPointWeightMap* FClothingMaskListItem::GetMask()
 
 FClothPhysicalMeshData* FClothingMaskListItem::GetMeshData()
 {
-	UClothingAssetCommon * Asset = ClothingAsset.Get();
-	return Asset && Asset->IsValidLod(LodIndex) ? &Asset->ClothLodData[LodIndex]->ClothPhysicalMeshData : nullptr;
+	UClothingAssetCommon* Asset = ClothingAsset.Get();
+	return Asset && Asset->IsValidLod(LodIndex) ? &Asset->LodData[LodIndex].PhysicalMeshData : nullptr;
 }
 
 USkeletalMesh* FClothingMaskListItem::GetOwningMesh()
 {
-	UClothingAssetCommon * Asset = ClothingAsset.Get();
+	UClothingAssetCommon* Asset = ClothingAsset.Get();
 	return Asset ? Cast<USkeletalMesh>(Asset->GetOuter()) : nullptr;
 }
 
@@ -308,26 +308,26 @@ private:
 
 			for(int32 CurrIndex = 0; CurrIndex < NumLods - 1; ++CurrIndex)
 			{
-				UClothLODDataCommon* SourceLod = Asset->ClothLodData[CurrIndex];
-				UClothLODDataCommon* DestLod = Asset->ClothLodData[CurrIndex + 1];
+				FClothLODDataCommon& SourceLod = Asset->LodData[CurrIndex];
+				FClothLODDataCommon& DestLod = Asset->LodData[CurrIndex + 1];
 
-				DestLod->ParameterMasks.Reset();
+				DestLod.PointWeightMaps.Reset();
 
-				for(FPointWeightMap& SourceMask : SourceLod->ParameterMasks)
+				for(FPointWeightMap& SourceMask : SourceLod.PointWeightMaps)
 				{
-					DestLod->ParameterMasks.AddDefaulted();
-					FPointWeightMap& DestMask = DestLod->ParameterMasks.Last();
+					DestLod.PointWeightMaps.AddDefaulted();
+					FPointWeightMap& DestMask = DestLod.PointWeightMaps.Last();
 
 					DestMask.Name = SourceMask.Name;
 					DestMask.bEnabled = SourceMask.bEnabled;
 					DestMask.CurrentTarget = SourceMask.CurrentTarget;
 
 					ClothingMeshUtils::FVertexParameterMapper ParameterMapper(
-						DestLod->ClothPhysicalMeshData.Vertices,
-						DestLod->ClothPhysicalMeshData.Normals,
-						SourceLod->ClothPhysicalMeshData.Vertices,
-						SourceLod->ClothPhysicalMeshData.Normals,
-						SourceLod->ClothPhysicalMeshData.Indices
+						DestLod.PhysicalMeshData.Vertices,
+						DestLod.PhysicalMeshData.Normals,
+						SourceLod.PhysicalMeshData.Vertices,
+						SourceLod.PhysicalMeshData.Normals,
+						SourceLod.PhysicalMeshData.Indices
 					);
 
 					ParameterMapper.Map(SourceMask.Values, DestMask.Values);
@@ -501,16 +501,15 @@ private:
 		);
 	}
 
-	UClothLODDataCommon* GetCurrentLod() const
+	FClothLODDataCommon* GetCurrentLod() const
 	{
 		if(Item.IsValid())
 		{
 			if(UClothingAssetCommon* Asset = Item->ClothingAsset.Get())
 			{
-				if(Asset->ClothLodData.IsValidIndex(Item->LodIndex))
+				if(Asset->LodData.IsValidIndex(Item->LodIndex))
 				{
-					UClothLODDataCommon* LodData = Asset->ClothLodData[Item->LodIndex];
-					return LodData;
+					return &Asset->LodData[Item->LodIndex];
 				}
 			}
 		}
@@ -527,11 +526,11 @@ private:
 			FScopedTransaction CurrTransaction(LOCTEXT("DeleteMask_Transaction", "Delete clothing parameter mask."));
 			Item->ClothingAsset->Modify();
 
-		if(UClothLODDataCommon* LodData = GetCurrentLod())
+		if(FClothLODDataCommon* LodData = GetCurrentLod())
 		{
-			if(LodData->ParameterMasks.IsValidIndex(Item->MaskIndex))
+			if(LodData->PointWeightMaps.IsValidIndex(Item->MaskIndex))
 			{
-				LodData->ParameterMasks.RemoveAt(Item->MaskIndex);
+				LodData->PointWeightMaps.RemoveAt(Item->MaskIndex);
 
 				// We've removed a mask, so it will need to be applied to the clothing data
 				if(Item.IsValid())
@@ -655,12 +654,12 @@ private:
 						// Disable all other masks that affect this target
 						if(UClothingAssetCommon* Asset = InItem->ClothingAsset.Get())
 						{
-							if(Asset->ClothLodData.IsValidIndex(InItem->LodIndex))
+							if(Asset->LodData.IsValidIndex(InItem->LodIndex))
 							{
-								UClothLODDataCommon* LodData = Asset->ClothLodData[InItem->LodIndex];
+								FClothLODDataCommon& LodData = Asset->LodData[InItem->LodIndex];
 
 								TArray<FPointWeightMap*> AllTargetMasks;
-								LodData->GetParameterMasksForTarget(Mask->CurrentTarget, AllTargetMasks);
+								LodData.GetParameterMasksForTarget(Mask->CurrentTarget, AllTargetMasks);
 
 								for(FPointWeightMap* TargetMask : AllTargetMasks)
 								{
@@ -1139,14 +1138,14 @@ FReply SClothAssetSelector::AddNewMask()
 {
 	if(UClothingAssetCommon* Asset = SelectedAsset.Get())
 	{
-		if(Asset->ClothLodData.IsValidIndex(SelectedLod))
+		if(Asset->LodData.IsValidIndex(SelectedLod))
 		{
-			UClothLODDataCommon* LodData = Asset->ClothLodData[SelectedLod];
-			const int32 NumRequiredValues = LodData->ClothPhysicalMeshData.Vertices.Num();
+			FClothLODDataCommon& LodData = Asset->LodData[SelectedLod];
+			const int32 NumRequiredValues = LodData.PhysicalMeshData.Vertices.Num();
 
-			LodData->ParameterMasks.AddDefaulted();
+			LodData.PointWeightMaps.AddDefaulted();
 
-			FPointWeightMap& NewMask = LodData->ParameterMasks.Last();
+			FPointWeightMap& NewMask = LodData.PointWeightMaps.Last();
 
 			NewMask.Name = TEXT("New Mask");
 			NewMask.CurrentTarget = (uint8)EWeightMapTargetCommon::None;
@@ -1245,8 +1244,8 @@ void SClothAssetSelector::RefreshMaskList()
 	UClothingAssetCommon* Asset = SelectedAsset.Get();
 	if(Asset && Asset->IsValidLod(SelectedLod))
 	{
-		UClothLODDataCommon* LodData = Asset->ClothLodData[SelectedLod];
-		const int32 NumMasks = LodData->ParameterMasks.Num();
+		const FClothLODDataCommon& LodData = Asset->LodData[SelectedLod];
+		const int32 NumMasks = LodData.PointWeightMaps.Num();
 
 		for(int32 Index = 0; Index < NumMasks; ++Index)
 		{
@@ -1295,11 +1294,11 @@ void SClothAssetSelector::OnClothingLodSelected(int32 InNewLod)
 		SetSelectedLod(InNewLod);
 
 		int32 NewMaskSelection = INDEX_NONE;
-		if(SelectedAsset->ClothLodData.IsValidIndex(SelectedLod))
+		if(SelectedAsset->LodData.IsValidIndex(SelectedLod))
 		{
-			UClothLODDataCommon* LodData = SelectedAsset->ClothLodData[SelectedLod];
+			const FClothLODDataCommon& LodData = SelectedAsset->LodData[SelectedLod];
 
-			if(LodData->ParameterMasks.Num() > 0)
+			if(LodData.PointWeightMaps.Num() > 0)
 			{
 				NewMaskSelection = 0;
 			}
@@ -1321,8 +1320,8 @@ void SClothAssetSelector::SetSelectedAsset(TWeakObjectPtr<UClothingAssetCommon>
 		{
 			SetSelectedLod(0);
 
-			UClothLODDataCommon* LodData = NewAsset->ClothLodData[SelectedLod];
-			if(LodData->ParameterMasks.Num() > 0)
+			const FClothLODDataCommon& LodData = NewAsset->LodData[SelectedLod];
+			if(LodData.PointWeightMaps.Num() > 0)
 			{
 				SetSelectedMask(0);
 			}
diff --git a/Engine/Source/Editor/ClothPainter/Private/SCopyVertexColorSettingsPanel.cpp b/Engine/Source/Editor/ClothPainter/Private/SCopyVertexColorSettingsPanel.cpp
--- a/Engine/Source/Editor/ClothPainter/Private/SCopyVertexColorSettingsPanel.cpp
+++ b/Engine/Source/Editor/ClothPainter/Private/SCopyVertexColorSettingsPanel.cpp
@@ -112,15 +112,15 @@ FReply SCopyVertexColorSettingsPanel::OnCopyClicked()
 	UClothingAssetCommon* Asset = SelectedAssetPtr.Get();
 	if (Asset && SelectedMask)
 	{
-		UClothLODDataCommon* ClothLODData = Asset->ClothLodData[SelectedLOD];
+		const FClothLODDataCommon& ClothLODData = Asset->LodData[SelectedLOD];
 
-		check(ClothLODData->ClothPhysicalMeshData.Vertices.Num() == ClothLODData->ClothPhysicalMeshData.VertexColors.Num());
-		int32 NumVerts = ClothLODData->ClothPhysicalMeshData.Vertices.Num();
+		check(ClothLODData.PhysicalMeshData.Vertices.Num() == ClothLODData.PhysicalMeshData.VertexColors.Num());
+		int32 NumVerts = ClothLODData.PhysicalMeshData.Vertices.Num();
 		check(SelectedMask->Values.Num() == NumVerts);
 
 		for (int32 VertIdx = 0; VertIdx < NumVerts; VertIdx++)
 		{
-			const FColor VertColor = ClothLODData->ClothPhysicalMeshData.VertexColors[VertIdx];
+			const FColor VertColor = ClothLODData.PhysicalMeshData.VertexColors[VertIdx];
 			SelectedMask->Values[VertIdx] = GetColorChannelAsFloat(VertColor, CopyParams.ColorChannel, CopyParams.ScalingFactor);
 		}
 	}
diff --git a/Engine/Source/Editor/ClothingSystemEditor/Private/ClothingAssetFactory.cpp b/Engine/Source/Editor/ClothingSystemEditor/Private/ClothingAssetFactory.cpp
--- a/Engine/Source/Editor/ClothingSystemEditor/Private/ClothingAssetFactory.cpp
+++ b/Engine/Source/Editor/ClothingSystemEditor/Private/ClothingAssetFactory.cpp
@@ -244,7 +244,7 @@ UClothingAssetBase* UClothingAssetFactory::CreateFromSkeletalMesh(USkeletalMesh*
 
 	// Adding a new LOD from this skeletal mesh
 	NewAsset->AddNewLod();
-	UClothLODDataCommon* LodData = NewAsset->ClothLodData.Last();
+	FClothLODDataCommon& LodData = NewAsset->LodData.Last();
 
 	if(ImportToLodInternal(TargetMesh, Params.LodIndex, Params.SourceSection, NewAsset, LodData))
 	{
@@ -324,20 +324,20 @@ UClothingAssetBase* UClothingAssetFactory::ImportLodToClothing(USkeletalMesh* Ta
 			// Everything looks good, continue to actual import
 			UClothingAssetCommon* ConcreteTarget = CastChecked<UClothingAssetCommon>(TargetClothing);
 
-			UClothLODDataCommon* RemapSource = nullptr;
+			const FClothLODDataCommon* RemapSource = nullptr;
 
 			if(Params.bRemapParameters)
 			{
 				if(Params.TargetLod == ConcreteTarget->GetNumLods())
 				{
 					// New LOD, remap from previous
-					RemapSource = ConcreteTarget->ClothLodData.Last();
+					RemapSource = &ConcreteTarget->LodData.Last();
 				}
 				else
 				{
 					// This is a replacement, remap from current LOD
-					check(ConcreteTarget->ClothLodData.IsValidIndex(Params.TargetLod));
-					RemapSource = ConcreteTarget->ClothLodData[Params.TargetLod];
+					check(ConcreteTarget->LodData.IsValidIndex(Params.TargetLod));
+					RemapSource = &ConcreteTarget->LodData[Params.TargetLod];
 				}
 			}
 
@@ -345,17 +345,17 @@ UClothingAssetBase* UClothingAssetFactory::ImportLodToClothing(USkeletalMesh* Ta
 			{
 				ConcreteTarget->AddNewLod();
 			}
-			else if(!ConcreteTarget->ClothLodData.IsValidIndex(Params.TargetLod))
+			else if(!ConcreteTarget->LodData.IsValidIndex(Params.TargetLod))
 			{
 				LogAndToastWarning(LOCTEXT("Warning_InvalidLodTarget", "Failed to import clothing LOD, invalid target LOD."));
 				return nullptr;
 			}
 
-			UClothLODDataCommon* NewLod = ConcreteTarget->ClothLodData[Params.TargetLod];
+			FClothLODDataCommon& NewLod = ConcreteTarget->LodData[Params.TargetLod];
 
 			if(Params.TargetLod > 0 && Params.bRemapParameters)
 			{
-				RemapSource = ConcreteTarget->ClothLodData[Params.TargetLod - 1];
+				RemapSource = &ConcreteTarget->LodData[Params.TargetLod - 1];
 			}
 
 			if(ImportToLodInternal(TargetMesh, Params.LodIndex, Params.SourceSection, ConcreteTarget, NewLod, RemapSource))
@@ -401,12 +401,7 @@ UClothingAssetBase* UClothingAssetFactory::CreateFromApexAsset(nvidia::apex::Clo
 	for(int32 CurrLodIdx = 0; CurrLodIdx < NumLodsToBuild; ++CurrLodIdx)
 	{
 		NewClothingAsset->AddNewLod();
-		UClothLODDataCommon* CurrentLodDataPtr = NewClothingAsset->ClothLodData[CurrLodIdx];
-		if (!ensure(CurrentLodDataPtr))
-		{
-			continue;
-		}
-		UClothLODDataCommon& CurrentLodData = *CurrentLodDataPtr;
+		FClothLODDataCommon& CurrentLodData = NewClothingAsset->LodData[CurrLodIdx];
 
 		TArray<FApexVertData> ApexVertData;
 
@@ -423,11 +418,11 @@ UClothingAssetBase* UClothingAssetFactory::CreateFromApexAsset(nvidia::apex::Clo
 		ConfigNv->ClothingWindMethod = EClothingWindMethodNv::Legacy;
 
 		// Fixup unreal-side bone indices
-		const int32 NumBoneDatas = CurrentLodData.ClothPhysicalMeshData.BoneData.Num();
+		const int32 NumBoneDatas = CurrentLodData.PhysicalMeshData.BoneData.Num();
 		check(NumBoneDatas == ApexVertData.Num());
 		for(int32 BoneDataIndex = 0; BoneDataIndex < NumBoneDatas; ++BoneDataIndex)
 		{
-			FClothVertBoneData& BoneData = CurrentLodData.ClothPhysicalMeshData.BoneData[BoneDataIndex];
+			FClothVertBoneData& BoneData = CurrentLodData.PhysicalMeshData.BoneData[BoneDataIndex];
 			FApexVertData& CurrentVertData = ApexVertData[BoneDataIndex];
 		
 			for(int32 BoneInfluenceIdx = 0; BoneInfluenceIdx < MAX_TOTAL_INFLUENCES; ++BoneInfluenceIdx)
@@ -474,18 +469,17 @@ UClothingAssetBase* UClothingAssetFactory::CreateFromApexAsset(nvidia::apex::Clo
 	NewClothingAsset->CalculateReferenceBoneIndex();
 
 	// Add masks for parameters
-	for (UClothLODDataCommon* LodPtr : NewClothingAsset->ClothLodData)
+	for (FClothLODDataCommon& LodData : NewClothingAsset->LodData)
 	{
-		check(LodPtr);
-		const FClothPhysicalMeshData& PhysMesh = LodPtr->ClothPhysicalMeshData;
+		const FClothPhysicalMeshData& PhysMesh = LodData.PhysicalMeshData;
 
 		// Didn't do anything previously - clear out in case there's something in there
 		// so we can use it correctly now.
-		LodPtr->ParameterMasks.Reset(3);
+		LodData.PointWeightMaps.Reset(3);
 
 		// Max distances
-		LodPtr->ParameterMasks.AddDefaulted();
-		FPointWeightMap& MaxDistanceMask = LodPtr->ParameterMasks.Last();
+		LodData.PointWeightMaps.AddDefaulted();
+		FPointWeightMap& MaxDistanceMask = LodData.PointWeightMaps.Last();
 		const FPointWeightMap& MaxDistances = PhysMesh.GetWeightMap(EWeightMapTargetCommon::MaxDistance);
 		MaxDistanceMask.Initialize(MaxDistances, EWeightMapTargetCommon::MaxDistance);
 
@@ -493,13 +487,13 @@ UClothingAssetBase* UClothingAssetFactory::CreateFromApexAsset(nvidia::apex::Clo
 		if (BackstopRadiuses && !BackstopRadiuses->IsZeroed())
 		{
 			// Backstop radii
-			LodPtr->ParameterMasks.AddDefaulted();
-			FPointWeightMap& BackstopRadiusMask = LodPtr->ParameterMasks.Last();
+			LodData.PointWeightMaps.AddDefaulted();
+			FPointWeightMap& BackstopRadiusMask = LodData.PointWeightMaps.Last();
 			BackstopRadiusMask.Initialize(*BackstopRadiuses, EWeightMapTargetCommon::BackstopRadius);
 
 			// Backstop distances
-			LodPtr->ParameterMasks.AddDefaulted();
-			FPointWeightMap& BackstopDistanceMask = LodPtr->ParameterMasks.Last();
+			LodData.PointWeightMaps.AddDefaulted();
+			FPointWeightMap& BackstopDistanceMask = LodData.PointWeightMaps.Last();
 			const FPointWeightMap& BackstopDistances = PhysMesh.GetWeightMap(EWeightMapTargetCommon::BackstopDistance);
 			BackstopDistanceMask.Initialize(BackstopDistances, EWeightMapTargetCommon::BackstopDistance);
 		}
@@ -849,7 +843,7 @@ void UClothingAssetFactory::ExtractBoneData(UClothingAssetCommon* NewAsset, Clot
 	}
 }
 
-void UClothingAssetFactory::ExtractSphereCollisions(UClothingAssetCommon* NewAsset, nvidia::apex::ClothingAsset &InApexAsset, int32 InLodIdx, UClothLODDataCommon &InLodData)
+void UClothingAssetFactory::ExtractSphereCollisions(UClothingAssetCommon* NewAsset, nvidia::apex::ClothingAsset &InApexAsset, int32 InLodIdx, FClothLODDataCommon& InLodData)
 {
 	const NvParameterized::Interface* AssetParams = InApexAsset.getAssetNvParameterized();
 
@@ -1132,8 +1126,8 @@ bool UClothingAssetFactory::ImportToLodInternal(
 	int32 SourceLodIndex, 
 	int32 SourceSectionIndex, 
 	UClothingAssetCommon* DestAsset, 
-	UClothLODDataCommon* DestLod, 
-	UClothLODDataCommon* InParameterRemapSource)
+	FClothLODDataCommon& DestLod, 
+	const FClothLODDataCommon* InParameterRemapSource)
 {
 	if(!SourceMesh || !SourceMesh->GetImportedModel())
 	{
@@ -1223,17 +1217,17 @@ bool UClothingAssetFactory::ImportToLodInternal(
 
 	if(InParameterRemapSource)
 	{
-		const FClothPhysicalMeshData& RemapPhysMesh = InParameterRemapSource->ClothPhysicalMeshData;
+		const FClothPhysicalMeshData& RemapPhysMesh = InParameterRemapSource->PhysicalMeshData;
 		CachedPositions = RemapPhysMesh.Vertices;
 		CachedNormals = RemapPhysMesh.Normals;
 		CachedIndices = RemapPhysMesh.Indices;
-		SourceMaskCopy = InParameterRemapSource->ParameterMasks;
+		SourceMaskCopy = InParameterRemapSource->PointWeightMaps;
 		NumSourceMasks = SourceMaskCopy.Num();
 
 		bPerformParamterRemap = true;
 	}
 
-	FClothPhysicalMeshData& PhysMesh = DestLod->ClothPhysicalMeshData;
+	FClothPhysicalMeshData& PhysMesh = DestLod.PhysicalMeshData;
 	PhysMesh.Reset(NumUniqueVerts, NumIndices);
 
 	for(int32 VertexIndex = 0; VertexIndex < NumUniqueVerts; ++VertexIndex)
@@ -1258,8 +1252,8 @@ bool UClothingAssetFactory::ImportToLodInternal(
 	}
 
 	// Add a max distance parameter mask to begin with
-	DestLod->ParameterMasks.AddDefaulted();
-	FPointWeightMap& Mask = DestLod->ParameterMasks.Last();
+	DestLod.PointWeightMaps.AddDefaulted();
+	FPointWeightMap& Mask = DestLod.PointWeightMaps.Last();
 	const FPointWeightMap& MaxDistances = PhysMesh.GetWeightMap(EWeightMapTargetCommon::MaxDistance);
 	Mask.Initialize(MaxDistances, EWeightMapTargetCommon::MaxDistance);
 
@@ -1293,14 +1287,14 @@ bool UClothingAssetFactory::ImportToLodInternal(
 	{
 		ClothingMeshUtils::FVertexParameterMapper ParameterRemapper(PhysMesh.Vertices, PhysMesh.Normals, CachedPositions, CachedNormals, CachedIndices);
 
-		DestLod->ParameterMasks.Reset(NumSourceMasks);
+		DestLod.PointWeightMaps.Reset(NumSourceMasks);
 
 		for(int32 MaskIndex = 0; MaskIndex < NumSourceMasks; ++MaskIndex)
 		{
 			const FPointWeightMap& SourceMask = SourceMaskCopy[MaskIndex];
 
-			DestLod->ParameterMasks.AddDefaulted();
-			FPointWeightMap& DestMask = DestLod->ParameterMasks.Last();
+			DestLod.PointWeightMaps.AddDefaulted();
+			FPointWeightMap& DestMask = DestLod.PointWeightMaps.Last();
 
 			DestMask.Initialize(PhysMesh.Vertices.Num());
 			DestMask.CurrentTarget = SourceMask.CurrentTarget;
@@ -1317,11 +1311,11 @@ bool UClothingAssetFactory::ImportToLodInternal(
 
 #if WITH_APEX_CLOTHING
 
-void UClothingAssetFactory::ExtractLodPhysicalData(UClothingAssetCommon* NewAsset, ClothingAsset &InApexAsset, int32 InLodIdx, UClothLODDataCommon &InLodData, TArray<FApexVertData>& OutApexVertData)
+void UClothingAssetFactory::ExtractLodPhysicalData(UClothingAssetCommon* NewAsset, ClothingAsset &InApexAsset, int32 InLodIdx, FClothLODDataCommon& InLodData, TArray<FApexVertData>& OutApexVertData)
 {
 	const NvParameterized::Interface* AssetParams = InApexAsset.getAssetNvParameterized();
 
-	FClothPhysicalMeshData& PhysData = InLodData.ClothPhysicalMeshData;
+	FClothPhysicalMeshData& PhysData = InLodData.PhysicalMeshData;
 	NvParameterized::Handle GraphicalMeshArrayHandle(*AssetParams, "graphicalLods");
 
 	int32 NumGraphicalLods = 0;
diff --git a/Engine/Source/Editor/ClothingSystemEditor/Public/ClothingAssetFactory.h b/Engine/Source/Editor/ClothingSystemEditor/Public/ClothingAssetFactory.h
--- a/Engine/Source/Editor/ClothingSystemEditor/Public/ClothingAssetFactory.h
+++ b/Engine/Source/Editor/ClothingSystemEditor/Public/ClothingAssetFactory.h
@@ -65,16 +65,16 @@ private:
 
 	// Extraction methods for pulling the required data from an APEX asset and
 	// pushing it to a UClothingAssetCommon
-	void ExtractLodPhysicalData(UClothingAssetCommon* NewAsset, nvidia::apex::ClothingAsset &InApexAsset, int32 InLodIdx, UClothLODDataCommon &InLodData, TArray<FApexVertData>& OutApexVertData);
-	void ExtractBoneData(UClothingAssetCommon* NewAsset, nvidia::apex::ClothingAsset &InApexAsset);
-	void ExtractSphereCollisions(UClothingAssetCommon* NewAsset, nvidia::apex::ClothingAsset &InApexAsset, int32 InLodIdx, UClothLODDataCommon &InLodData);
-	void ExtractMaterialParameters(UClothingAssetCommon* NewAsset, nvidia::apex::ClothingAsset &InApexAsset);
+	void ExtractLodPhysicalData(UClothingAssetCommon* NewAsset, nvidia::apex::ClothingAsset& InApexAsset, int32 InLodIdx, FClothLODDataCommon& InLodData, TArray<FApexVertData>& OutApexVertData);
+	void ExtractBoneData(UClothingAssetCommon* NewAsset, nvidia::apex::ClothingAsset& InApexAsset);
+	void ExtractSphereCollisions(UClothingAssetCommon* NewAsset, nvidia::apex::ClothingAsset& InApexAsset, int32 InLodIdx, FClothLODDataCommon& InLodData);
+	void ExtractMaterialParameters(UClothingAssetCommon* NewAsset, nvidia::apex::ClothingAsset& InApexAsset);
 #endif
 
 	// Utility methods for skeletal mesh extraction //////////////////////////
 
 	/** Handles internal import of LODs */
-	bool ImportToLodInternal(USkeletalMesh* SourceMesh, int32 SourceLodIndex, int32 SourceSectionIndex, UClothingAssetCommon* DestAsset, UClothLODDataCommon* DestLod, UClothLODDataCommon* InParameterRemapSource = nullptr);
+	bool ImportToLodInternal(USkeletalMesh* SourceMesh, int32 SourceLodIndex, int32 SourceSectionIndex, UClothingAssetCommon* DestAsset, FClothLODDataCommon& DestLod, const FClothLODDataCommon* InParameterRemapSource = nullptr);
 
 	//////////////////////////////////////////////////////////////////////////
 
diff --git a/Engine/Source/Editor/Persona/Private/PersonaMeshDetails.cpp b/Engine/Source/Editor/Persona/Private/PersonaMeshDetails.cpp
--- a/Engine/Source/Editor/Persona/Private/PersonaMeshDetails.cpp
+++ b/Engine/Source/Editor/Persona/Private/PersonaMeshDetails.cpp
@@ -5525,9 +5525,9 @@ TSharedRef<SUniformGridPanel> FPersonaMeshDetails::MakeClothingDetailsWidget(int
 
 		if (UClothingAssetCommon* Asset = Cast<UClothingAssetCommon>(ClothingAsset))
 		{
-			const UClothLODDataCommon* const LodData = Asset->ClothLodData[LODIndex];
-			const FClothPhysicalMeshData& PhysMeshData = LodData->ClothPhysicalMeshData;
-			const FClothCollisionData& CollisionData = LodData->CollisionData;
+			const FClothLODDataCommon& LodData = Asset->LodData[LODIndex];
+			const FClothPhysicalMeshData& PhysMeshData = LodData.PhysicalMeshData;
+			const FClothCollisionData& CollisionData = LodData.CollisionData;
 
 			Grid->AddSlot(0, RowNumber)
 				.HAlign(HAlign_Center)
diff --git a/Engine/Source/Editor/UnrealEd/Private/Animation/DebugSkelMeshComponent.cpp b/Engine/Source/Editor/UnrealEd/Private/Animation/DebugSkelMeshComponent.cpp
--- a/Engine/Source/Editor/UnrealEd/Private/Animation/DebugSkelMeshComponent.cpp
+++ b/Engine/Source/Editor/UnrealEd/Private/Animation/DebugSkelMeshComponent.cpp
@@ -742,8 +742,8 @@ void UDebugSkelMeshComponent::RebuildClothingSectionsFixedVerts()
 				if(BaseAsset)
 				{
 					UClothingAssetCommon* ConcreteAsset = Cast<UClothingAssetCommon>(BaseAsset);
-					UClothLODDataCommon* LodData = ConcreteAsset->ClothLodData[Section.ClothingData.AssetLodIndex];
-					const FPointWeightMap& MaxDistances = LodData->ClothPhysicalMeshData.GetWeightMap(EWeightMapTargetCommon::MaxDistance);
+					const FClothLODDataCommon& LodData = ConcreteAsset->LodData[Section.ClothingData.AssetLodIndex];
+					const FPointWeightMap& MaxDistances = LodData.PhysicalMeshData.GetWeightMap(EWeightMapTargetCommon::MaxDistance);
 
 					for(FMeshToMeshVertData& VertData : Section.ClothMappingData)
 					{
@@ -821,7 +821,7 @@ void UDebugSkelMeshComponent::RefreshSelectedClothingSkinnedPositions()
 		{
 			UClothingAssetCommon* ConcreteAsset = Cast<UClothingAssetCommon>(*Asset);
 
-			if(ConcreteAsset->ClothLodData.IsValidIndex(SelectedClothingLodForPainting))
+			if(ConcreteAsset->LodData.IsValidIndex(SelectedClothingLodForPainting))
 			{
 				SkinnedSelectedClothingPositions.Reset();
 				SkinnedSelectedClothingNormals.Reset();
@@ -830,9 +830,9 @@ void UDebugSkelMeshComponent::RefreshSelectedClothingSkinnedPositions()
 				// Pass LOD0 to collect all bones
 				GetCurrentRefToLocalMatrices(RefToLocals, 0);
 
-				UClothLODDataCommon* LodData = ConcreteAsset->ClothLodData[SelectedClothingLodForPainting];
+				const FClothLODDataCommon& LodData = ConcreteAsset->LodData[SelectedClothingLodForPainting];
 
-				ClothingMeshUtils::SkinPhysicsMesh(ConcreteAsset->UsedBoneIndices, LodData->ClothPhysicalMeshData, FTransform::Identity, RefToLocals.GetData(), RefToLocals.Num(), SkinnedSelectedClothingPositions, SkinnedSelectedClothingNormals);
+				ClothingMeshUtils::SkinPhysicsMesh(ConcreteAsset->UsedBoneIndices, LodData.PhysicalMeshData, FTransform::Identity, RefToLocals.GetData(), RefToLocals.Num(), SkinnedSelectedClothingPositions, SkinnedSelectedClothingNormals);
 				RebuildCachedClothBounds();
 			}
 		}
@@ -1035,15 +1035,15 @@ FDebugSkelMeshDynamicData::FDebugSkelMeshDynamicData(UDebugSkelMeshComponent* In
 
 					if(UClothingAssetCommon* ConcreteAsset = Cast<UClothingAssetCommon>(BaseAsset))
 					{
-						if(ConcreteAsset->ClothLodData.IsValidIndex(InComponent->SelectedClothingLodForPainting))
+						if(ConcreteAsset->LodData.IsValidIndex(InComponent->SelectedClothingLodForPainting))
 						{
-							UClothLODDataCommon* LodData = ConcreteAsset->ClothLodData[InComponent->SelectedClothingLodForPainting];
+							const FClothLODDataCommon& LodData = ConcreteAsset->LodData[InComponent->SelectedClothingLodForPainting];
 
-							ClothingSimIndices = LodData->ClothPhysicalMeshData.Indices;
+							ClothingSimIndices = LodData.PhysicalMeshData.Indices;
 
-							if(LodData->ParameterMasks.IsValidIndex(InComponent->SelectedClothingLodMaskForPainting))
+							if(LodData.PointWeightMaps.IsValidIndex(InComponent->SelectedClothingLodMaskForPainting))
 							{
-								FPointWeightMap& Mask = LodData->ParameterMasks[InComponent->SelectedClothingLodMaskForPainting];
+								const FPointWeightMap& Mask = LodData.PointWeightMaps[InComponent->SelectedClothingLodMaskForPainting];
 
 								ClothingVisiblePropertyValues = Mask.Values;
 							}
diff --git a/Engine/Source/Runtime/ClothingSystemRuntimeCommon/Private/ClothLODData.cpp b/Engine/Source/Runtime/ClothingSystemRuntimeCommon/Private/ClothLODData.cpp
--- a/Engine/Source/Runtime/ClothingSystemRuntimeCommon/Private/ClothLODData.cpp
+++ b/Engine/Source/Runtime/ClothingSystemRuntimeCommon/Private/ClothLODData.cpp
@@ -2,22 +2,14 @@
 
 #include "ClothLODData.h"
 #include "ClothPhysicalMeshData.h"
-
-UClothLODDataCommon::UClothLODDataCommon(const FObjectInitializer& Init)
-	: Super(Init)
-	, PhysicalMeshData_DEPRECATED(nullptr)
-{
-}
-
-UClothLODDataCommon::~UClothLODDataCommon()
-{}
+#include "ClothingAssetCustomVersion.h"
 
 #if WITH_EDITORONLY_DATA
-void UClothLODDataCommon::GetParameterMasksForTarget(
+void FClothLODDataCommon::GetParameterMasksForTarget(
 	const uint8 InTarget, 
 	TArray<FPointWeightMap*>& OutMasks)
 {
-	for(FPointWeightMap& Mask : ParameterMasks)
+	for(FPointWeightMap& Mask : PointWeightMaps)
 	{
 		if(Mask.CurrentTarget == InTarget)
 		{
@@ -27,47 +19,63 @@ void UClothLODDataCommon::GetParameterMasksForTarget(
 }
 #endif // WITH_EDITORONLY_DATA
 
-void UClothLODDataCommon::Serialize(FArchive& Ar)
+bool FClothLODDataCommon::Serialize(FArchive& Ar)
 {
-	Super::Serialize(Ar);
-
-	// Ryan - do we need to do this tagged property business still?
+	Ar.UsingCustomVersion(FClothingAssetCustomVersion::GUID);
 
-	// Serialize normal tagged data
-	//UScriptStruct* Struct = UClothLODDataNv::StaticStruct();
-	//UClass* Class = UClothLODDataNv::StaticClass();
-	/*if (!Ar.IsCountingMemory())
+	// Serialize normal tagged property data
+	if (!Ar.IsCountingMemory())
 	{
-		//Struct->SerializeTaggedProperties(Ar, (uint8*)this, Struct, nullptr);
-		Class->SerializeTaggedProperties(Ar, (uint8*)this, Class, nullptr);
-	}*/
+		UScriptStruct* const Struct = FClothLODDataCommon::StaticStruct();
+		Struct->SerializeTaggedProperties(Ar, (uint8*)this, Struct, nullptr);
+	}
 
 	// Serialize the mesh to mesh data (not a USTRUCT)
 	Ar << TransitionUpSkinData
 	   << TransitionDownSkinData;
-}
-
-void UClothLODDataCommon::PostLoad()
-{
-	Super::PostLoad();
 
-	if (PhysicalMeshData_DEPRECATED)
+	const int32 ClothingCustomVersion = Ar.CustomVer(FClothingAssetCustomVersion::GUID);
+	if (ClothingCustomVersion < FClothingAssetCustomVersion::MovePropertiesToCommonBaseClasses)
 	{
-		PhysicalMeshData_DEPRECATED->ConditionalPostLoad();  // Makes sure the UObject has finished loading
-		ClothPhysicalMeshData.MigrateFrom(PhysicalMeshData_DEPRECATED);
-		PhysicalMeshData_DEPRECATED = nullptr;
+		// Migrate maps
+		PhysicalMeshData.GetWeightMap(EWeightMapTargetCommon::MaxDistance).Values = MoveTemp(PhysicalMeshData.MaxDistances_DEPRECATED);
+		PhysicalMeshData.GetWeightMap(EWeightMapTargetCommon::BackstopDistance).Values = MoveTemp(PhysicalMeshData.BackstopDistances_DEPRECATED);
+		PhysicalMeshData.GetWeightMap(EWeightMapTargetCommon::BackstopRadius).Values = MoveTemp(PhysicalMeshData.BackstopRadiuses_DEPRECATED);
+		PhysicalMeshData.GetWeightMap(EWeightMapTargetCommon::AnimDriveMultiplier).Values = MoveTemp(PhysicalMeshData.AnimDriveMultipliers_DEPRECATED);
+
+#if WITH_EDITORONLY_DATA
+		// Migrate editor maps
+		PointWeightMaps.SetNum(ParameterMasks_DEPRECATED.Num());
+		for (int32 i = 0; i < PointWeightMaps.Num(); ++i)
+		{
+			ParameterMasks_DEPRECATED[i].MigrateTo(PointWeightMaps[i]);
+		}
+		ParameterMasks_DEPRECATED.Empty();
+#endif // WITH_EDITORONLY_DATA
+
+#if WITH_CHAOS
+		// Rebuild surface points so that the legacy Apex convex collision data can also be used with Chaos
+		for (auto& Convex : CollisionData.Convexes)
+		{
+			if (!Convex.SurfacePoints.Num())
+			{
+				Convex.RebuildSurfacePoints();
+			}
+		}
+#endif // #if WITH_CHAOS
 	}
+	return true;
 }
 
 #if WITH_EDITOR
-void UClothLODDataCommon::PushWeightsToMesh()
+void FClothLODDataCommon::PushWeightsToMesh()
 {
-	ClothPhysicalMeshData.ClearWeightMaps();
-	for (const FPointWeightMap& Weights : ParameterMasks)
+	PhysicalMeshData.ClearWeightMaps();
+	for (const FPointWeightMap& Weights : PointWeightMaps)
 	{
 		if (Weights.bEnabled)
 		{
-			FPointWeightMap& TargetWeightMap = ClothPhysicalMeshData.FindOrAddWeightMap(Weights.CurrentTarget);
+			FPointWeightMap& TargetWeightMap = PhysicalMeshData.FindOrAddWeightMap(Weights.CurrentTarget);
 			TargetWeightMap.Values = Weights.Values;
 		}
 	}
diff --git a/Engine/Source/Runtime/ClothingSystemRuntimeCommon/Private/ClothLODData_Legacy.cpp b/Engine/Source/Runtime/ClothingSystemRuntimeCommon/Private/ClothLODData_Legacy.cpp
--- a/Engine/Source/Runtime/ClothingSystemRuntimeCommon/Private/ClothLODData_Legacy.cpp
+++ b/Engine/Source/Runtime/ClothingSystemRuntimeCommon/Private/ClothLODData_Legacy.cpp
@@ -1,8 +1,8 @@
 // Copyright Epic Games, Inc. All Rights Reserved.
 
 #include "ClothLODData_Legacy.h"
-#include "PointWeightMap.h"
 #include "ClothLODData.h"
+#include "ClothPhysicalMeshDataBase_Legacy.h"
 
 FClothParameterMask_Legacy::FClothParameterMask_Legacy()
 	: MaskName(NAME_None)
@@ -12,51 +12,73 @@ FClothParameterMask_Legacy::FClothParameterMask_Legacy()
 	, bEnabled(false)
 {}
 
-void FClothParameterMask_Legacy::MigrateTo(FPointWeightMap* Weights) const
+void FClothParameterMask_Legacy::MigrateTo(FPointWeightMap& PointWeightMap)
 {
-	Weights->Values = Values;
+	PointWeightMap.Values = MoveTemp(Values);
 #if WITH_EDITORONLY_DATA
-	Weights->Name = MaskName;
-	Weights->CurrentTarget = static_cast<uint8>(CurrentTarget);
-	Weights->bEnabled = bEnabled;
+	PointWeightMap.Name = MaskName;
+	PointWeightMap.CurrentTarget = static_cast<uint8>(CurrentTarget);
+	PointWeightMap.bEnabled = bEnabled;
 #endif
 }
 
-bool FClothLODData_Legacy::Serialize(FArchive& Ar)
+UClothLODDataCommon_Legacy::UClothLODDataCommon_Legacy(const FObjectInitializer& Init)
+	: Super(Init)
+	, PhysicalMeshData_DEPRECATED(nullptr)
 {
-	// Serialize normal tagged data
-	if (!Ar.IsCountingMemory())
+}
+
+UClothLODDataCommon_Legacy::~UClothLODDataCommon_Legacy()
+{}
+
+void UClothLODDataCommon_Legacy::Serialize(FArchive& Ar)
+{
+	Super::Serialize(Ar);
+
+	// Serialize the mesh to mesh data (not a USTRUCT)
+	Ar << TransitionUpSkinData
+	   << TransitionDownSkinData;
+}
+
+void UClothLODDataCommon_Legacy::PostLoad()
+{
+	Super::PostLoad();
+
+	if (PhysicalMeshData_DEPRECATED)
 	{
-		UScriptStruct* Struct = FClothLODData_Legacy::StaticStruct();
-		Struct->SerializeTaggedProperties(Ar, (uint8*)this, Struct, nullptr);
+		PhysicalMeshData_DEPRECATED->ConditionalPostLoad();  // Makes sure the UObject has finished loading
+		ClothPhysicalMeshData.MigrateFrom(PhysicalMeshData_DEPRECATED);
+		PhysicalMeshData_DEPRECATED = nullptr;
 	}
-	// Serialize the mesh to mesh data (not a USTRUCT)
-	Ar	<< TransitionUpSkinData
-		<< TransitionDownSkinData;
-	return true;
 }
 
-void FClothLODData_Legacy::MigrateTo(UClothLODDataCommon* LodData) const
+void UClothLODDataCommon_Legacy::MigrateTo(FClothLODDataCommon& LodData)
 {
-	PhysicalMeshData.MigrateTo(LodData->ClothPhysicalMeshData);
-	LodData->CollisionData = CollisionData;
-#if WITH_CHAOS
-	// Rebuild surface points so that the legacy Apex convex collision data can also be used with Chaos
-	for (auto& Convex : LodData->CollisionData.Convexes)
+	// Migrate mesh data
+	if (PhysicalMeshData_DEPRECATED)
 	{
-		if (!Convex.SurfacePoints.Num())
-		{
-			Convex.RebuildSurfacePoints();
-		}
+		// Migrate PhysicalMeshData from when it was a UObject
+		PhysicalMeshData_DEPRECATED->ConditionalPostLoad();  // Makes sure the UObject has finished loading
+		LodData.PhysicalMeshData.MigrateFrom(PhysicalMeshData_DEPRECATED);
 	}
-#endif // #if WITH_CHAOS
-#if WITH_EDITORONLY_DATA
-	LodData->ParameterMasks.SetNum(ParameterMasks.Num());
-	for (int i = 0; i < ParameterMasks.Num(); i++)
+	else
 	{
-		ParameterMasks[i].MigrateTo(&LodData->ParameterMasks[i]);
+		// Migrate PhysicalMeshData from when it was a UStruct
+		LodData.PhysicalMeshData.MigrateFrom(ClothPhysicalMeshData);
 	}
+
+	// Migrate collision
+	LodData.CollisionData.Spheres = MoveTemp(CollisionData.Spheres);
+	LodData.CollisionData.SphereConnections = MoveTemp(CollisionData.SphereConnections);
+	LodData.CollisionData.Convexes = MoveTemp(CollisionData.Convexes);
+	LodData.CollisionData.Boxes = MoveTemp(CollisionData.Boxes);
+
+#if WITH_EDITORONLY_DATA
+	// Migrate editor maps
+	LodData.PointWeightMaps = MoveTemp(ParameterMasks);
 #endif // WITH_EDITORONLY_DATA
-	LodData->TransitionUpSkinData = TransitionUpSkinData;
-	LodData->TransitionDownSkinData = TransitionDownSkinData;
+
+	// Migrate skinning data
+	LodData.TransitionUpSkinData = MoveTemp(TransitionUpSkinData);
+	LodData.TransitionDownSkinData = MoveTemp(TransitionDownSkinData);
 }
diff --git a/Engine/Source/Runtime/ClothingSystemRuntimeCommon/Private/ClothPhysicalMeshData.cpp b/Engine/Source/Runtime/ClothingSystemRuntimeCommon/Private/ClothPhysicalMeshData.cpp
--- a/Engine/Source/Runtime/ClothingSystemRuntimeCommon/Private/ClothPhysicalMeshData.cpp
+++ b/Engine/Source/Runtime/ClothingSystemRuntimeCommon/Private/ClothPhysicalMeshData.cpp
@@ -10,60 +10,45 @@ FClothPhysicalMeshData::FClothPhysicalMeshData()
 	ClearWeightMaps();
 }
 
-void FClothPhysicalMeshData::MigrateTo(FClothPhysicalMeshData& ClothPhysicalMeshData) const
+void FClothPhysicalMeshData::MigrateFrom(FClothPhysicalMeshData& ClothPhysicalMeshData)
 {
-	// Migrate properties
-	ClothPhysicalMeshData.Vertices = Vertices;
-	ClothPhysicalMeshData.Normals = Normals;
+	if (this != &ClothPhysicalMeshData)
+	{
+		Vertices = MoveTemp(ClothPhysicalMeshData.Vertices);
+		Normals = MoveTemp(ClothPhysicalMeshData.Normals);
 #if WITH_EDITORONLY_DATA
-	ClothPhysicalMeshData.VertexColors = VertexColors;
+		VertexColors = MoveTemp(ClothPhysicalMeshData.VertexColors);
 #endif
-	ClothPhysicalMeshData.Indices = Indices;
-	ClothPhysicalMeshData.InverseMasses = InverseMasses;
-	ClothPhysicalMeshData.BoneData = BoneData;
-	ClothPhysicalMeshData.NumFixedVerts = NumFixedVerts;
-	ClothPhysicalMeshData.MaxBoneWeights = MaxBoneWeights;
-	ClothPhysicalMeshData.SelfCollisionIndices = SelfCollisionIndices;
-
-	// Migrate maps
-	if (MaxDistances_DEPRECATED.Num())
-	{
-		ClothPhysicalMeshData.GetWeightMap(EWeightMapTargetCommon::MaxDistance).Values = MaxDistances_DEPRECATED;
-	}
-	if (BackstopDistances_DEPRECATED.Num())
-	{
-		ClothPhysicalMeshData.GetWeightMap(EWeightMapTargetCommon::BackstopDistance).Values = BackstopDistances_DEPRECATED;
-	}
-	if (BackstopRadiuses_DEPRECATED.Num())
-	{
-		ClothPhysicalMeshData.GetWeightMap(EWeightMapTargetCommon::BackstopRadius).Values = BackstopRadiuses_DEPRECATED;
-	}
-	if (AnimDriveMultipliers_DEPRECATED.Num())
-	{
-		ClothPhysicalMeshData.GetWeightMap(EWeightMapTargetCommon::AnimDriveMultiplier).Values = AnimDriveMultipliers_DEPRECATED;
+		Indices = MoveTemp(ClothPhysicalMeshData.Indices);
+		WeightMaps = MoveTemp(ClothPhysicalMeshData.WeightMaps);
+		InverseMasses = MoveTemp(ClothPhysicalMeshData.InverseMasses);
+		BoneData = MoveTemp(ClothPhysicalMeshData.BoneData);
+		NumFixedVerts = ClothPhysicalMeshData.NumFixedVerts;
+		MaxBoneWeights = ClothPhysicalMeshData.MaxBoneWeights;
+		SelfCollisionIndices = MoveTemp(ClothPhysicalMeshData.SelfCollisionIndices);
 	}
 }
 
-void FClothPhysicalMeshData::MigrateFrom(const UClothPhysicalMeshDataBase_Legacy* ClothPhysicalMeshDataBase)
+void FClothPhysicalMeshData::MigrateFrom(UClothPhysicalMeshDataBase_Legacy* ClothPhysicalMeshDataBase)
 {
-	Vertices = ClothPhysicalMeshDataBase->Vertices;
-	Normals = ClothPhysicalMeshDataBase->Normals;
+	Vertices = MoveTemp(ClothPhysicalMeshDataBase->Vertices);
+	Normals = MoveTemp(ClothPhysicalMeshDataBase->Normals);
 #if WITH_EDITORONLY_DATA
-	VertexColors = ClothPhysicalMeshDataBase->VertexColors;
+	VertexColors = MoveTemp(ClothPhysicalMeshDataBase->VertexColors);
 #endif
-	Indices = ClothPhysicalMeshDataBase->Indices;
-	InverseMasses = ClothPhysicalMeshDataBase->InverseMasses;
-	BoneData = ClothPhysicalMeshDataBase->BoneData;
+	Indices = MoveTemp(ClothPhysicalMeshDataBase->Indices);
+	InverseMasses = MoveTemp(ClothPhysicalMeshDataBase->InverseMasses);
+	BoneData = MoveTemp(ClothPhysicalMeshDataBase->BoneData);
 	NumFixedVerts = ClothPhysicalMeshDataBase->NumFixedVerts;
 	MaxBoneWeights = ClothPhysicalMeshDataBase->MaxBoneWeights;
-	SelfCollisionIndices = ClothPhysicalMeshDataBase->SelfCollisionIndices;
+	SelfCollisionIndices = MoveTemp(ClothPhysicalMeshDataBase->SelfCollisionIndices);
 
 	const TArray<uint32> FloatArrayIds = ClothPhysicalMeshDataBase->GetFloatArrayIds();
 	for (uint32 FloatArrayId : FloatArrayIds)
 	{
-		if (const TArray<float>* const FloatArray = ClothPhysicalMeshDataBase->GetFloatArray(FloatArrayId))
+		if (TArray<float>* const FloatArray = ClothPhysicalMeshDataBase->GetFloatArray(FloatArrayId))
 		{
-			FindOrAddWeightMap(FloatArrayId).Values = *FloatArray;
+			FindOrAddWeightMap(FloatArrayId).Values = MoveTemp(*FloatArray);
 		}
 	}
 }
diff --git a/Engine/Source/Runtime/ClothingSystemRuntimeCommon/Private/ClothingAsset.cpp b/Engine/Source/Runtime/ClothingSystemRuntimeCommon/Private/ClothingAsset.cpp
--- a/Engine/Source/Runtime/ClothingSystemRuntimeCommon/Private/ClothingAsset.cpp
+++ b/Engine/Source/Runtime/ClothingSystemRuntimeCommon/Private/ClothingAsset.cpp
@@ -211,7 +211,7 @@ bool UClothingAssetCommon::BindToSkeletalMesh(
 	}
 
 	// If we don't have clothing data
-	if(!ClothLodData.IsValidIndex(InAssetLodIndex))
+	if(!LodData.IsValidIndex(InAssetLodIndex))
 	{
 		FText Error = FText::Format(
 			LOCTEXT("Error_NoClothingLod", "Failed to bind clothing asset {0} LOD{1} as LOD{2} does not exist."), 
@@ -270,7 +270,7 @@ bool UClothingAssetCommon::BindToSkeletalMesh(
 	CalculateReferenceBoneIndex();
 
 	// Grab the clothing and skel lod data
-	UClothLODDataCommon* LodData = ClothLodData[InAssetLodIndex];
+	FClothLODDataCommon& ClothLodData = LodData[InAssetLodIndex];
 	FSkeletalMeshLODModel& SkelLod = InSkelMesh->GetImportedModel()->LODModels[InMeshLodIndex];
 
 	FSkelMeshSection& OriginalSection = SkelLod.Sections[InSectionIndex];
@@ -298,9 +298,9 @@ bool UClothingAssetCommon::BindToSkeletalMesh(
 
 	ClothingMeshUtils::ClothMeshDesc TargetMesh(RenderPositions, RenderNormals, IndexView);
 	ClothingMeshUtils::ClothMeshDesc SourceMesh(
-		LodData->ClothPhysicalMeshData.Vertices, 
-		LodData->ClothPhysicalMeshData.Normals, 
-		LodData->ClothPhysicalMeshData.Indices);
+		ClothLodData.PhysicalMeshData.Vertices, 
+		ClothLodData.PhysicalMeshData.Normals, 
+		ClothLodData.PhysicalMeshData.Indices);
 
 	ClothingMeshUtils::GenerateMeshToMeshSkinningData(MeshToMeshData, TargetMesh, &RenderTangents, SourceMesh);
 
@@ -312,7 +312,7 @@ bool UClothingAssetCommon::BindToSkeletalMesh(
 	}
 
 	// Calculate fixed verts
-	const FPointWeightMap& MaxDistances = LodData->ClothPhysicalMeshData.GetWeightMap(EWeightMapTargetCommon::MaxDistance);
+	const FPointWeightMap& MaxDistances = ClothLodData.PhysicalMeshData.GetWeightMap(EWeightMapTargetCommon::MaxDistance);
 	for(FMeshToMeshVertData& VertData : MeshToMeshData)
 	{
 		if(MaxDistances.AreAnyBelowThreshold(
@@ -517,9 +517,9 @@ void UClothingAssetCommon::ForEachInteractorUsingClothing(TFunction<void(UClothi
 
 void UClothingAssetCommon::ApplyParameterMasks()
 {
-	for(UClothLODDataCommon* Lod : ClothLodData)
+	for(FClothLODDataCommon& Lod : LodData)
 	{
-		Lod->PushWeightsToMesh();
+		Lod.PushWeightsToMesh();
 	}
 	InvalidateCachedData();
 }
@@ -532,11 +532,11 @@ void UClothingAssetCommon::BuildLodTransitionData()
 		const bool bHasPrevLod = LodIndex > 0;
 		const bool bHasNextLod = LodIndex < NumLods - 1;
 
-		UClothLODDataCommon* CurrentLod = ClothLodData[LodIndex];
-		const FClothPhysicalMeshData& CurrentPhysMesh = CurrentLod->ClothPhysicalMeshData;
+		FClothLODDataCommon& CurrentLod = LodData[LodIndex];
+		const FClothPhysicalMeshData& CurrentPhysMesh = CurrentLod.PhysicalMeshData;
 
-		UClothLODDataCommon* PrevLod = bHasPrevLod ? ClothLodData[LodIndex - 1] : nullptr;
-		UClothLODDataCommon* NextLod = bHasNextLod ? ClothLodData[LodIndex + 1] : nullptr;
+		FClothLODDataCommon* const PrevLod = bHasPrevLod ? &LodData[LodIndex - 1] : nullptr;
+		FClothLODDataCommon* const NextLod = bHasNextLod ? &LodData[LodIndex + 1] : nullptr;
 
 		const int32 CurrentLodNumVerts = CurrentPhysMesh.Vertices.Num();
 
@@ -544,17 +544,17 @@ void UClothingAssetCommon::BuildLodTransitionData()
 
 		if(PrevLod)
 		{
-			FClothPhysicalMeshData& PrevPhysMesh = PrevLod->ClothPhysicalMeshData;
-			CurrentLod->TransitionUpSkinData.Empty(CurrentLodNumVerts);
+			FClothPhysicalMeshData& PrevPhysMesh = PrevLod->PhysicalMeshData;
+			CurrentLod.TransitionUpSkinData.Empty(CurrentLodNumVerts);
 			ClothingMeshUtils::ClothMeshDesc PrevMeshDesc(PrevPhysMesh.Vertices, PrevPhysMesh.Normals, PrevPhysMesh.Indices);
-			ClothingMeshUtils::GenerateMeshToMeshSkinningData(CurrentLod->TransitionUpSkinData, CurrentMeshDesc, nullptr, PrevMeshDesc);
+			ClothingMeshUtils::GenerateMeshToMeshSkinningData(CurrentLod.TransitionUpSkinData, CurrentMeshDesc, nullptr, PrevMeshDesc);
 		}
 		if(NextLod)
 		{
-			FClothPhysicalMeshData& NextPhysMesh = NextLod->ClothPhysicalMeshData;
-			CurrentLod->TransitionDownSkinData.Empty(CurrentLodNumVerts);
+			FClothPhysicalMeshData& NextPhysMesh = NextLod->PhysicalMeshData;
+			CurrentLod.TransitionDownSkinData.Empty(CurrentLodNumVerts);
 			ClothingMeshUtils::ClothMeshDesc NextMeshDesc(NextPhysMesh.Vertices, NextPhysMesh.Normals, NextPhysMesh.Indices);
-			ClothingMeshUtils::GenerateMeshToMeshSkinningData(CurrentLod->TransitionDownSkinData, CurrentMeshDesc, nullptr, NextMeshDesc);
+			ClothingMeshUtils::GenerateMeshToMeshSkinningData(CurrentLod.TransitionDownSkinData, CurrentMeshDesc, nullptr, NextMeshDesc);
 		}
 	}
 }
@@ -602,10 +602,9 @@ void UClothingAssetCommon::CalculateReferenceBoneIndex()
 		// First build a list per used bone for it's path to root
 		TArray<int32> WeightedBones;  // List of actually weighted (not just used) bones
 
-		for(UClothLODDataCommon* CurLod : ClothLodData)
+		for(FClothLODDataCommon& CurLod : LodData)
 		{
-			check(CurLod);
-			const FClothPhysicalMeshData& MeshData = CurLod->ClothPhysicalMeshData;
+			const FClothPhysicalMeshData& MeshData = CurLod.PhysicalMeshData;
 			for(const FClothVertBoneData& VertBoneData : MeshData.BoneData)
 			{
 				for(int32 InfluenceIndex = 0; InfluenceIndex < MAX_TOTAL_INFLUENCES; ++InfluenceIndex)
@@ -693,22 +692,21 @@ void UClothingAssetCommon::CalculateReferenceBoneIndex()
 
 bool UClothingAssetCommon::IsValidLod(int32 InLodIndex) const
 {
-	return ClothLodData.IsValidIndex(InLodIndex);
+	return LodData.IsValidIndex(InLodIndex);
 }
 
 int32 UClothingAssetCommon::GetNumLods() const
 {
-	return ClothLodData.Num();
+	return LodData.Num();
 }
 
 void UClothingAssetCommon::BuildSelfCollisionData()
 {
 	if (ClothConfigs.Num())
 	{
-		for(UClothLODDataCommon* Lod : ClothLodData)
+		for(FClothLODDataCommon& Lod : LodData)
 		{
-			check(Lod);
-			Lod->ClothPhysicalMeshData.BuildSelfCollisionData(ClothConfigs);
+			Lod.PhysicalMeshData.BuildSelfCollisionData(ClothConfigs);
 		}
 	}
 }
@@ -716,63 +714,57 @@ void UClothingAssetCommon::BuildSelfCollisionData()
 void UClothingAssetCommon::PostLoad()
 {
 	Super::PostLoad();
+
+	// Migrate the deprecated UObject based lod class to the non-UObject lod structure to prevent PostLoad dependency issues
 	// TODO: Remove all UObject PostLoad dependencies.
 	//       Even with these ConditionalPostLoad calls, the UObject PostLoads' order of execution cannot be guaranteed.
-	for (UClothLODDataCommon* Lod : ClothLodData)
+	for (UClothLODDataCommon_Legacy* LodDeprecated : ClothLodData_DEPRECATED)
 	{
-		Lod->ConditionalPostLoad();
-	}
-
-	const int32 AnimPhysCustomVersion = GetLinkerCustomVersion(FAnimPhysObjectVersion::GUID);
-	const int32 ClothingCustomVersion = GetLinkerCustomVersion(FClothingAssetCustomVersion::GUID);
-
-	if (ClothingCustomVersion < FClothingAssetCustomVersion::MovePropertiesToCommonBaseClasses)
-	{
-		// Remap legacy struct FClothLODData to class UClothLODDataCommon
-		for (const FClothLODData_Legacy& ClothLODData_Legacy : LodData_DEPRECATED)
+		if (LodDeprecated)
 		{
+			LodDeprecated->ConditionalPostLoad();
+
 			const int32 Idx = AddNewLod();
-			ClothLODData_Legacy.MigrateTo(ClothLodData[Idx]);
+			LodDeprecated->MigrateTo(LodData[Idx]);
 		}
-		LodData_DEPRECATED.Empty();
 	}
-	if(AnimPhysCustomVersion < FAnimPhysObjectVersion::AddedClothingMaskWorkflow)
+	ClothLodData_DEPRECATED.Empty();
+
+	const int32 AnimPhysCustomVersion = GetLinkerCustomVersion(FAnimPhysObjectVersion::GUID);
+	if (AnimPhysCustomVersion < FAnimPhysObjectVersion::AddedClothingMaskWorkflow)
 	{
 #if WITH_EDITORONLY_DATA
 		// Convert current parameters to masks
-		for (UClothLODDataCommon* LodPtr : ClothLodData)
+		for (FClothLODDataCommon& Lod : LodData)
 		{
-			check(LodPtr);
-			UClothLODDataCommon& Lod = *LodPtr;
-			const FClothPhysicalMeshData& PhysMesh = Lod.ClothPhysicalMeshData;
+			const FClothPhysicalMeshData& PhysMesh = Lod.PhysicalMeshData;
 
 			// Didn't do anything previously - clear out in case there's something in there
 			// so we can use it correctly now.
-			Lod.ParameterMasks.Reset(3);
+			Lod.PointWeightMaps.Reset(3);
 
 			// Max distances (Always present)
-			Lod.ParameterMasks.AddDefaulted();
-			FPointWeightMap& MaxDistanceMask = Lod.ParameterMasks.Last();
+			Lod.PointWeightMaps.AddDefaulted();
+			FPointWeightMap& MaxDistanceMask = Lod.PointWeightMaps.Last();
 			const FPointWeightMap& MaxDistances = PhysMesh.GetWeightMap(EWeightMapTargetCommon::MaxDistance);
 			MaxDistanceMask.Initialize(MaxDistances, EWeightMapTargetCommon::MaxDistance);
 
 			// Following params are only added if necessary, if we don't have any backstop
 			// radii then there's no backstops.
 			const FPointWeightMap* const BackstopRadiuses = PhysMesh.FindWeightMap(EWeightMapTargetCommon::BackstopRadius);
-			if(BackstopRadiuses && !BackstopRadiuses->IsZeroed())
+			if (BackstopRadiuses && !BackstopRadiuses->IsZeroed())
 			{
 				// Backstop radii
-				Lod.ParameterMasks.AddDefaulted();
-				FPointWeightMap& BackstopRadiusMask = Lod.ParameterMasks.Last();
+				Lod.PointWeightMaps.AddDefaulted();
+				FPointWeightMap& BackstopRadiusMask = Lod.PointWeightMaps.Last();
 				BackstopRadiusMask.Initialize(*BackstopRadiuses, EWeightMapTargetCommon::BackstopRadius);
 
 				// Backstop distances
-				Lod.ParameterMasks.AddDefaulted();
-				FPointWeightMap& BackstopDistanceMask = Lod.ParameterMasks.Last();
+				Lod.PointWeightMaps.AddDefaulted();
+				FPointWeightMap& BackstopDistanceMask = Lod.PointWeightMaps.Last();
 				const FPointWeightMap& BackstopDistances = PhysMesh.GetWeightMap(EWeightMapTargetCommon::BackstopDistance);
 				BackstopDistanceMask.Initialize(BackstopDistances, EWeightMapTargetCommon::BackstopDistance);
 			}
-			
 		}
 #endif
 
@@ -782,17 +774,18 @@ void UClothingAssetCommon::PostLoad()
 
 #if WITH_EDITORONLY_DATA
 	// Fix content imported before we kept vertex colors
+	const int32 ClothingCustomVersion = GetLinkerCustomVersion(FClothingAssetCustomVersion::GUID);
 	if (ClothingCustomVersion < FClothingAssetCustomVersion::AddVertexColorsToPhysicalMesh)
 	{
-		for (UClothLODDataCommon* Lod : ClothLodData)
+		for (FClothLODDataCommon& Lod : LodData)
 		{
-			const int32 NumVerts = Lod->ClothPhysicalMeshData.Vertices.Num(); // number of verts
+			const int32 NumVerts = Lod.PhysicalMeshData.Vertices.Num(); // number of verts
 
-			Lod->ClothPhysicalMeshData.VertexColors.Reset();
-			Lod->ClothPhysicalMeshData.VertexColors.AddUninitialized(NumVerts);
+			Lod.PhysicalMeshData.VertexColors.Reset();
+			Lod.PhysicalMeshData.VertexColors.AddUninitialized(NumVerts);
 			for (int32 VertIdx = 0; VertIdx < NumVerts; VertIdx++)
 			{
-				Lod->ClothPhysicalMeshData.VertexColors[VertIdx] = FColor::White;
+				Lod.PhysicalMeshData.VertexColors[VertIdx] = FColor::White;
 			}
 		}
 	}
@@ -828,6 +821,7 @@ void UClothingAssetCommon::PostLoad()
 		{
 			if (UClothConfigCommon* const ClothConfigCommon = Cast<UClothConfigCommon>(ClothConfig.Value))
 			{
+				ClothConfigCommon->ConditionalPostLoad();
 				ClothConfigCommon->MigrateFrom(ClothConfig_DEPRECATED);
 			}
 		}
@@ -848,6 +842,7 @@ void UClothingAssetCommon::PostLoad()
 					{
 						if (UClothConfigCommon* const ClothConfigCommon = Cast<UClothConfigCommon>(ClothConfig.Value))
 						{
+							ClothConfigCommon->ConditionalPostLoad();
 							ClothConfigCommon->MigrateFrom(ClothConfigLegacy);
 						}
 					}
@@ -1027,11 +1022,10 @@ void UClothingAssetCommon::PostUpdateAllAssets()
 
 void UClothingAssetCommon::InvalidateCachedData()
 {
-	for(UClothLODDataCommon* CurrentLodData : ClothLodData)
+	for(FClothLODDataCommon& CurrentLodData : LodData)
 	{
-		check(CurrentLodData);
 		// Recalculate inverse masses for the physical mesh particles
-		FClothPhysicalMeshData& PhysMesh = CurrentLodData->ClothPhysicalMeshData;
+		FClothPhysicalMeshData& PhysMesh = CurrentLodData.PhysicalMeshData;
 		check(PhysMesh.Indices.Num() % 3 == 0);
 
 		TArray<float>& InvMasses = PhysMesh.InverseMasses;
@@ -1125,9 +1119,7 @@ void UClothingAssetCommon::InvalidateCachedData()
 
 int32 UClothingAssetCommon::AddNewLod()
 {
-	const int32 Idx = ClothLodData.AddDefaulted();
-	ClothLodData[Idx] = NewObject<UClothLODDataCommon>(this, UClothLODDataCommon::StaticClass(), NAME_None, RF_Transactional);
-	return Idx;
+	return LodData.AddDefaulted();
 }
 
 void UClothingAssetCommon::PostEditChangeChainProperty(FPropertyChangedChainEvent& ChainEvent)
diff --git a/Engine/Source/Runtime/ClothingSystemRuntimeCommon/Public/ClothLODData.h b/Engine/Source/Runtime/ClothingSystemRuntimeCommon/Public/ClothLODData.h
--- a/Engine/Source/Runtime/ClothingSystemRuntimeCommon/Public/ClothLODData.h
+++ b/Engine/Source/Runtime/ClothingSystemRuntimeCommon/Public/ClothLODData.h
@@ -2,38 +2,30 @@
 #pragma once
 
 #include "ClothPhysicalMeshData.h"
-#include "ClothPhysicalMeshDataBase_Legacy.h"
-#include "ClothCollisionData.h"
-#include "PointWeightMap.h"
-#include "SkeletalMeshTypes.h"
-
+#include "ClothLODData_Legacy.h"
 #include "ClothLODData.generated.h"
 
 /** Common Cloth LOD representation for all clothing assets. */
-UCLASS()
-class CLOTHINGSYSTEMRUNTIMECOMMON_API UClothLODDataCommon : public UObject
+USTRUCT()
+struct CLOTHINGSYSTEMRUNTIMECOMMON_API FClothLODDataCommon
 {
 	GENERATED_BODY()
-public:
-	UClothLODDataCommon(const FObjectInitializer& Init);
-	virtual ~UClothLODDataCommon();
-
-	// Deprecated, use ClothPhysicalMeshData instead
-	UPROPERTY()
-	UClothPhysicalMeshDataBase_Legacy* PhysicalMeshData_DEPRECATED;
 
 	// Raw phys mesh data
 	UPROPERTY(EditAnywhere, Category = SimMesh)
-	FClothPhysicalMeshData ClothPhysicalMeshData;
+	FClothPhysicalMeshData PhysicalMeshData;
 
-	// Collision primitive and covex data for clothing collisions
+	// Collision primitive and convex data for clothing collisions
 	UPROPERTY(EditAnywhere, Category = Collision)
 	FClothCollisionData CollisionData;
 
 #if WITH_EDITORONLY_DATA
+	UPROPERTY()
+	TArray<FClothParameterMask_Legacy> ParameterMasks_DEPRECATED;
+
 	// Parameter masks defining the physics mesh masked data
 	UPROPERTY(EditAnywhere, Category = Masks)
-	TArray<FPointWeightMap> ParameterMasks;
+	TArray<FPointWeightMap> PointWeightMaps;
 
 	// Get all available parameter masks for the specified target
 	void GetParameterMasksForTarget(const uint8 InTarget, TArray<FPointWeightMap*>& OutMasks);
@@ -42,6 +34,7 @@ public:
 	/** Copy \c ParameterMasks to corresponding targets in \c ClothPhysicalMeshData. */
 	void PushWeightsToMesh();
 #endif
+
 	// Skinning data for transitioning from a higher detail LOD to this one
 	TArray<FMeshToMeshVertData> TransitionUpSkinData;
 
@@ -49,8 +42,14 @@ public:
 	TArray<FMeshToMeshVertData> TransitionDownSkinData;
 
 	// Custom serialize for transition
-	virtual void Serialize(FArchive& Ar) override;
+	bool Serialize(FArchive& Ar);
+};
 
-	// Migrate deprecated properties
-	virtual void PostLoad() override;
+template<>
+struct TStructOpsTypeTraits<FClothLODDataCommon> : public TStructOpsTypeTraitsBase2<FClothLODDataCommon>
+{
+	enum
+	{
+		WithSerializer = true,
+	};
 };
diff --git a/Engine/Source/Runtime/ClothingSystemRuntimeCommon/Public/ClothLODData_Legacy.h b/Engine/Source/Runtime/ClothingSystemRuntimeCommon/Public/ClothLODData_Legacy.h
--- a/Engine/Source/Runtime/ClothingSystemRuntimeCommon/Public/ClothLODData_Legacy.h
+++ b/Engine/Source/Runtime/ClothingSystemRuntimeCommon/Public/ClothLODData_Legacy.h
@@ -1,15 +1,15 @@
 // Copyright Epic Games, Inc. All Rights Reserved.
 #pragma once
 
+
+#include "PointWeightMap.h"
 #include "ClothCollisionData.h"
 #include "SkeletalMeshTypes.h"
-#include "ClothVertBoneData.h"
 #include "ClothPhysicalMeshData.h"
-
 #include "ClothLODData_Legacy.generated.h"
 
 class UClothLODDataCommon;
-struct FPointWeightMap;
+class UClothPhysicalMeshDataBase_Legacy;
 
 /**
  * Deprecated, legacy definition kept for backward compatibility only.
@@ -23,7 +23,7 @@ struct CLOTHINGSYSTEMRUNTIMECOMMON_API FClothParameterMask_Legacy
 
 	FClothParameterMask_Legacy();
 
-	void MigrateTo(FPointWeightMap* Weights) const;
+	void MigrateTo(FPointWeightMap& Weights);
 
 	/** Name of the mask, mainly for users to differentiate */
 	UPROPERTY()
@@ -52,26 +52,32 @@ struct CLOTHINGSYSTEMRUNTIMECOMMON_API FClothParameterMask_Legacy
 
 /**
  * Deprecated, legacy definition kept for backward compatibility only.
- * Use UClothLODDataCommon instead.
- * Redirected from the now defunct ClothingSystemRuntime module.
+ * Use FClothLODDataCommon instead.
  */
-USTRUCT()
-struct CLOTHINGSYSTEMRUNTIMECOMMON_API FClothLODData_Legacy
+UCLASS()
+class CLOTHINGSYSTEMRUNTIMECOMMON_API UClothLODDataCommon_Legacy : public UObject
 {
 	GENERATED_BODY()
+public:
+	UClothLODDataCommon_Legacy(const FObjectInitializer& Init);
+	virtual ~UClothLODDataCommon_Legacy();
+
+	// Deprecated, use ClothPhysicalMeshData instead
+	UPROPERTY()
+	UClothPhysicalMeshDataBase_Legacy* PhysicalMeshData_DEPRECATED;
 
 	// Raw phys mesh data
 	UPROPERTY()
-	FClothPhysicalMeshData PhysicalMeshData;
+	FClothPhysicalMeshData ClothPhysicalMeshData;
 
-	// Collision primitive and covex data for clothing collisions
+	// Collision primitive and convex data for clothing collisions
 	UPROPERTY()
 	FClothCollisionData CollisionData;
 
 #if WITH_EDITORONLY_DATA
 	// Parameter masks defining the physics mesh masked data
 	UPROPERTY()
-	TArray<FClothParameterMask_Legacy> ParameterMasks;
+	TArray<FPointWeightMap> ParameterMasks;
 #endif // WITH_EDITORONLY_DATA
 
 	// Skinning data for transitioning from a higher detail LOD to this one
@@ -80,17 +86,12 @@ struct CLOTHINGSYSTEMRUNTIMECOMMON_API FClothLODData_Legacy
 	// Skinning data for transitioning from a lower detail LOD to this one
 	TArray<FMeshToMeshVertData> TransitionDownSkinData;
 
-	bool Serialize(FArchive& Ar);
+	// Custom serialize for transition
+	virtual void Serialize(FArchive& Ar) override;
 
-	// Migrate this legacy Nv struct to the new class format (called by UClothingAssetCommon::Postload())
-	void MigrateTo(UClothLODDataCommon* LodData) const;
-};
+	// Migrate deprecated properties
+	virtual void PostLoad() override;
 
-template<>
-struct TStructOpsTypeTraits<FClothLODData_Legacy> : public TStructOpsTypeTraitsBase2<FClothLODData_Legacy>
-{
-	enum
-	{
-		WithSerializer = true,
-	};
+	// Migrate this deprecated UObject class to the structure format (called by UClothingAssetCommon::PostLoad())
+	void MigrateTo(struct FClothLODDataCommon& LodData);
 };
diff --git a/Engine/Source/Runtime/ClothingSystemRuntimeCommon/Public/ClothPhysicalMeshData.h b/Engine/Source/Runtime/ClothingSystemRuntimeCommon/Public/ClothPhysicalMeshData.h
--- a/Engine/Source/Runtime/ClothingSystemRuntimeCommon/Public/ClothPhysicalMeshData.h
+++ b/Engine/Source/Runtime/ClothingSystemRuntimeCommon/Public/ClothPhysicalMeshData.h
@@ -20,11 +20,11 @@ struct CLOTHINGSYSTEMRUNTIMECOMMON_API FClothPhysicalMeshData
 	/** Construct an empty cloth physical mesh with default common targets. */
 	FClothPhysicalMeshData();
 
-	/** Migrate away the deprecated properties. */
-	void MigrateTo(FClothPhysicalMeshData& ClothPhysicalMeshData) const;
+	/** Migrate from same, used to migrate LOD data from the UClothLODDataCommon_Legacy. */
+	void MigrateFrom(FClothPhysicalMeshData& ClothPhysicalMeshData);
 
-	/** Migrate from the nv legacy physical mesh data class. */
-	void MigrateFrom(const UClothPhysicalMeshDataBase_Legacy* ClothPhysicalMeshDataBase);
+	/** Migrate from the legacy physical mesh data class, used to migrate LOD data from the UClothLODDataCommon_Legacy. */
+	void MigrateFrom(UClothPhysicalMeshDataBase_Legacy* ClothPhysicalMeshDataBase);
 
 	/** Reset the default common targets for this cloth physical mesh. */
 	void Reset(const int32 InNumVerts, const int32 InNumIndices);
diff --git a/Engine/Source/Runtime/ClothingSystemRuntimeCommon/Public/ClothingAsset.h b/Engine/Source/Runtime/ClothingSystemRuntimeCommon/Public/ClothingAsset.h
--- a/Engine/Source/Runtime/ClothingSystemRuntimeCommon/Public/ClothingAsset.h
+++ b/Engine/Source/Runtime/ClothingSystemRuntimeCommon/Public/ClothingAsset.h
@@ -213,9 +213,13 @@ public:
 	UPROPERTY()
 	UClothConfigBase* ChaosClothSimConfig_DEPRECATED;
 
+	// Deprecated. Use LodData instead
+	UPROPERTY()
+	TArray<UClothLODDataCommon_Legacy*> ClothLodData_DEPRECATED;
+
 	// The actual asset data, listed by LOD.
 	UPROPERTY()
-	TArray<UClothLODDataCommon*> ClothLodData;
+	TArray<FClothLODDataCommon> LodData;
 
 	// Tracks which clothing LOD each skel mesh LOD corresponds to (LodMap[SkelLod]=ClothingLod).
 	UPROPERTY()
@@ -244,13 +248,6 @@ public:
 	UPROPERTY()
 	FClothConfig_Legacy ClothConfig_DEPRECATED;
 
-	/** 
-	 * Deprecated property for transitioning \c FClothLODData struct to the 
-	 * \c UClothLODDataCommon class, in a new property called \c ClothLodData.
-	 */
-	UPROPERTY()
-	TArray<FClothLODData_Legacy> LodData_DEPRECATED;
-
 private:
 
 	// Add or replace a new cloth config of the specified type.
diff --git a/Engine/Source/Runtime/ClothingSystemRuntimeNv/Private/ClothingSimulationNv.cpp b/Engine/Source/Runtime/ClothingSystemRuntimeNv/Private/ClothingSimulationNv.cpp
--- a/Engine/Source/Runtime/ClothingSystemRuntimeNv/Private/ClothingSimulationNv.cpp
+++ b/Engine/Source/Runtime/ClothingSystemRuntimeNv/Private/ClothingSimulationNv.cpp
@@ -113,17 +113,17 @@ void FClothingSimulationNv::CreateActor(USkeletalMeshComponent* InOwnerComponent
 
 	// Need the current reftolocals so we can skin the ref pose for the sim mesh
 	TArray<FMatrix> RefToLocals;
-	InOwnerComponent->GetCurrentRefToLocalMatrices(RefToLocals, FMath::Min(InOwnerComponent->PredictedLODLevel, Asset->ClothLodData.Num() - 1));
+	InOwnerComponent->GetCurrentRefToLocalMatrices(RefToLocals, FMath::Min(InOwnerComponent->PredictedLODLevel, Asset->LodData.Num() - 1));
 
 	Actors.AddDefaulted();
 	FClothingActorNv& NewActor = Actors.Last();
 
 	NewActor.SimDataIndex = InSimDataIndex;
 
-	for(int32 LodIndex = 0; LodIndex < Asset->ClothLodData.Num(); ++LodIndex)
+	for(int32 LodIndex = 0; LodIndex < Asset->LodData.Num(); ++LodIndex)
 	{
-		UClothLODDataCommon& AssetLodData = *Asset->ClothLodData[LodIndex];
-		const FClothPhysicalMeshData& PhysMesh = AssetLodData.ClothPhysicalMeshData;
+		const FClothLODDataCommon& AssetLodData = Asset->LodData[LodIndex];
+		const FClothPhysicalMeshData& PhysMesh = AssetLodData.PhysicalMeshData;
 
 		NewActor.LodData.AddDefaulted();
 		FClothingActorNv::FActorLodData& ActorLodData = NewActor.LodData.Last();
@@ -563,7 +563,7 @@ void FClothingSimulationNv::Simulate(IClothingSimulationContext* InContext)
 
 		// To build motion constraints (max distances) we need to skin the entire physics mesh
 		// this call also updates our fixed particles to avoid iterating the particle list a second time
-		const FClothPhysicalMeshData& PhysMesh = Actor.AssetCreatedFrom->ClothLodData[Actor.CurrentLodIndex]->ClothPhysicalMeshData;
+		const FClothPhysicalMeshData& PhysMesh = Actor.AssetCreatedFrom->LodData[Actor.CurrentLodIndex].PhysicalMeshData;
 		ClothingMeshUtils::SkinPhysicsMesh(Actor.AssetCreatedFrom->UsedBoneIndices, PhysMesh, RootBoneTransform, NvContext->RefToLocals.GetData(), NvContext->RefToLocals.Num(), Actor.SkinnedPhysicsMeshPositions[Actor.CurrentSkinnedPositionIndex], Actor.SkinnedPhysicsMeshNormals);
 
 		nv::cloth::Cloth* CurrentCloth = Actor.LodData[Actor.CurrentLodIndex].Cloth;
@@ -751,7 +751,7 @@ void FClothingSimulationNv::ComputePhysicalMeshNormals(FClothingActorNv &Actor)
 
 	const uint32 NumParticles = Actor.LodData[CurrentClothingLod].Cloth->getNumParticles();
 	nv::cloth::MappedRange<physx::PxVec4> Particles = Actor.LodData[CurrentClothingLod].Cloth->getCurrentParticles();
-	const TArray<uint32>& Indices = Actor.AssetCreatedFrom->ClothLodData[CurrentClothingLod]->ClothPhysicalMeshData.Indices;
+	const TArray<uint32>& Indices = Actor.AssetCreatedFrom->LodData[CurrentClothingLod].PhysicalMeshData.Indices;
 	const uint32 NumIndices = Indices.Num();
 
 	// Using the face normals, calculate normals. These will not be normalized as we're adding together
@@ -911,7 +911,7 @@ void FClothingSimulationNv::GetCollisions(FClothCollisionData& OutCollisions, bo
 		else
 		{
 			check(Actor.AssetCreatedFrom);
-			OutCollisions.Append(Actor.AssetCreatedFrom->ClothLodData[Actor.CurrentLodIndex]->CollisionData);
+			OutCollisions.Append(Actor.AssetCreatedFrom->LodData[Actor.CurrentLodIndex].CollisionData);
 		}
 	}
 }
@@ -928,7 +928,7 @@ void FClothingSimulationNv::GatherStats() const
 			continue;
 		}
 
-		INC_DWORD_STAT_BY(STAT_NumClothVerts, Actor.AssetCreatedFrom->ClothLodData[Actor.CurrentLodIndex]->ClothPhysicalMeshData.Vertices.Num());
+		INC_DWORD_STAT_BY(STAT_NumClothVerts, Actor.AssetCreatedFrom->LodData[Actor.CurrentLodIndex].PhysicalMeshData.Vertices.Num());
 	}
 }
 
@@ -1107,14 +1107,14 @@ void FClothingSimulationNv::UpdateLod(int32 InPredictedLod, const FTransform& Co
 				{
 					// Reposition particles skinned to outgoing LOD
 					bool bLodTransitionUp = OldClothingLod < PredictedClothingLod;
-					UClothLODDataCommon& NewLodAssetData = *Actor.AssetCreatedFrom->ClothLodData[PredictedClothingLod];
-					TArray<FMeshToMeshVertData>& SkinData = bLodTransitionUp ? NewLodAssetData.TransitionUpSkinData : NewLodAssetData.TransitionDownSkinData;
+					const FClothLODDataCommon& NewLodAssetData = Actor.AssetCreatedFrom->LodData[PredictedClothingLod];
+					const TArray<FMeshToMeshVertData>& SkinData = bLodTransitionUp ? NewLodAssetData.TransitionUpSkinData : NewLodAssetData.TransitionDownSkinData;
 
 					for(int32 ParticleIndex = 0; ParticleIndex < NumNewParticles; ++ParticleIndex)
 					{
 						// Do some simple skinning, we only care about positions for this as particles are just
 						// positions inside the solver.
-						FMeshToMeshVertData& VertData = SkinData[ParticleIndex];
+						const FMeshToMeshVertData& VertData = SkinData[ParticleIndex];
 
 						const FVector A = P2UVector(OldLodParticles[VertData.SourceMeshVertIndices[0]]);
 						const FVector B = P2UVector(OldLodParticles[VertData.SourceMeshVertIndices[1]]);
@@ -1178,7 +1178,7 @@ void FClothingSimulationNv::UpdateLod(int32 InPredictedLod, const FTransform& Co
 					if(CSTransforms.Num() > Actor.AssetCreatedFrom->ReferenceBoneIndex)
 					{
 						// compute skinned positions to init sim mesh
-						const FClothPhysicalMeshData& PhysMesh = Actor.AssetCreatedFrom->ClothLodData[PredictedClothingLod]->ClothPhysicalMeshData;
+						const FClothPhysicalMeshData& PhysMesh = Actor.AssetCreatedFrom->LodData[PredictedClothingLod].PhysicalMeshData;
 						TArray<FVector> SkinnedPhysicsMeshPositions;
 						TArray<FVector> SkinnedPhysicsMeshNormals;
 						ClothingMeshUtils::SkinPhysicsMesh(
@@ -1266,7 +1266,7 @@ void FClothingSimulationNv::DebugDraw_PhysMesh(USkeletalMeshComponent* OwnerComp
 			Particles = NvClothSupport::CreateRange(ActorData.Px_RestPositions);
 		}
 
-		const FClothPhysicalMeshData& PhysMesh = Actor.AssetCreatedFrom->ClothLodData[CurrentClothLod]->ClothPhysicalMeshData;
+		const FClothPhysicalMeshData& PhysMesh = Actor.AssetCreatedFrom->LodData[CurrentClothLod].PhysicalMeshData;
 		const TArray<uint32>& Indices = PhysMesh.Indices;
 		const FPointWeightMap& MaxDistances = PhysMesh.GetWeightMap(EWeightMapTargetCommon::MaxDistance);
 
@@ -1419,7 +1419,7 @@ void FClothingSimulationNv::DebugDraw_Backstops(USkeletalMeshComponent* OwnerCom
 		{
 			continue;
 		}
-		const FClothPhysicalMeshData& PhysMesh = Asset->ClothLodData[Actor.CurrentLodIndex]->ClothPhysicalMeshData;
+		const FClothPhysicalMeshData& PhysMesh = Asset->LodData[Actor.CurrentLodIndex].PhysicalMeshData;
 		const FPointWeightMap* const BackstopDistances = PhysMesh.FindWeightMap(EWeightMapTargetCommon::BackstopDistance);
 		const FPointWeightMap* const BackstopRadiuses = PhysMesh.FindWeightMap(EWeightMapTargetCommon::BackstopRadius);
 		const FPointWeightMap* const MaxDistances = PhysMesh.FindWeightMap(EWeightMapTargetCommon::MaxDistance);
@@ -1493,7 +1493,7 @@ void FClothingSimulationNv::DebugDraw_MaxDistances(USkeletalMeshComponent* Owner
 		{
 			continue;
 		}
-		const FClothPhysicalMeshData& PhysMesh = Asset->ClothLodData[Actor.CurrentLodIndex]->ClothPhysicalMeshData;
+		const FClothPhysicalMeshData& PhysMesh = Asset->LodData[Actor.CurrentLodIndex].PhysicalMeshData;
 		const FPointWeightMap* const MaxDistances = PhysMesh.FindWeightMap(EWeightMapTargetCommon::MaxDistance);
 		if (!MaxDistances || !MaxDistances->Num())
 		{
@@ -1540,8 +1540,8 @@ void FClothingSimulationNv::DebugDraw_SelfCollision(USkeletalMeshComponent* Owne
 
 		const float SelfCollisionThickness = Config->SelfCollisionRadius;
 
-		const UClothLODDataCommon& LodData = *Asset->ClothLodData[Actor.CurrentLodIndex];
-		const FClothPhysicalMeshData& PhysMesh = LodData.ClothPhysicalMeshData;
+		const FClothLODDataCommon& LodData = Asset->LodData[Actor.CurrentLodIndex];
+		const FClothPhysicalMeshData& PhysMesh = LodData.PhysicalMeshData;
 
 		nv::cloth::Cloth* CurrentCloth = Actor.LodData[Actor.CurrentLodIndex].Cloth;
 		check(CurrentCloth);
@@ -1570,7 +1570,7 @@ void FClothingSimulationNv::DebugDraw_AnimDrive(USkeletalMeshComponent* OwnerCom
 		}
 
 		const UClothingAssetCommon* Asset = Actor.AssetCreatedFrom;
-		const FClothPhysicalMeshData& PhysMesh = Asset->ClothLodData[Actor.CurrentLodIndex]->ClothPhysicalMeshData;
+		const FClothPhysicalMeshData& PhysMesh = Asset->LodData[Actor.CurrentLodIndex].PhysicalMeshData;
 		const TArray<uint32>& Indices = PhysMesh.Indices;
 		FTransform RootBoneTransform = OwnerComponent->GetComponentSpaceTransforms()[Actor.AssetCreatedFrom->ReferenceBoneIndex];
 
@@ -1639,7 +1639,7 @@ FClothingActorNv::FClothingActorNv()
 
 void FClothingActorNv::SkinPhysicsMesh(FClothingSimulationContextNv* InContext)
 {
-	const FClothPhysicalMeshData& PhysMesh = AssetCreatedFrom->ClothLodData[CurrentLodIndex]->ClothPhysicalMeshData;
+	const FClothPhysicalMeshData& PhysMesh = AssetCreatedFrom->LodData[CurrentLodIndex].PhysicalMeshData;
 	FTransform RootBoneTransform = InContext->BoneTransforms[AssetCreatedFrom->ReferenceBoneIndex];
 	ClothingMeshUtils::SkinPhysicsMesh(AssetCreatedFrom->UsedBoneIndices, PhysMesh, RootBoneTransform, InContext->RefToLocals.GetData(), InContext->RefToLocals.Num(), SkinnedPhysicsMeshPositions[CurrentSkinnedPositionIndex], SkinnedPhysicsMeshNormals);
 }
@@ -1655,7 +1655,7 @@ void FClothingActorNv::UpdateMotionConstraints(FClothingSimulationContextNv* InC
 	nv::cloth::Cloth* const CurrentCloth = LodData[CurrentLodIndex].Cloth;
 	check(CurrentCloth);
 
-	const FClothPhysicalMeshData& PhysMesh = AssetCreatedFrom->ClothLodData[CurrentLodIndex]->ClothPhysicalMeshData;
+	const FClothPhysicalMeshData& PhysMesh = AssetCreatedFrom->LodData[CurrentLodIndex].PhysicalMeshData;
 
 	const FPointWeightMap* const MaxDistances = PhysMesh.FindWeightMap(EWeightMapTargetCommon::MaxDistance);
 	if (MaxDistances && MaxDistances->Num())
@@ -1702,7 +1702,7 @@ void FClothingActorNv::UpdateWind(FClothingSimulationContextNv* InContext, const
 			TArray<FVector>& ParticleVelocities = Scratch.ParticleVelocities;
 			CalculateParticleVelocities(ParticleVelocities);
 
-			const FClothPhysicalMeshData& PhysMesh = AssetCreatedFrom->ClothLodData[CurrentLodIndex]->ClothPhysicalMeshData;
+			const FClothPhysicalMeshData& PhysMesh = AssetCreatedFrom->LodData[CurrentLodIndex].PhysicalMeshData;
 			const FPointWeightMap* const MaxDistances = PhysMesh.FindWeightMap(EWeightMapTargetCommon::MaxDistance);
 			if (!MaxDistances || !MaxDistances->Num())
 			{
@@ -1757,7 +1757,7 @@ void FClothingActorNv::ConditionalRebuildCollisions()
 	AggregatedCollisions.Reset();
 
 	// Asset-embedded collisions (created during import)
-	AggregatedCollisions.Append(AssetCreatedFrom->ClothLodData[CurrentLodIndex]->CollisionData);
+	AggregatedCollisions.Append(AssetCreatedFrom->LodData[CurrentLodIndex].CollisionData);
 	// Extracted collisions from the physics asset selected by the user
 	AggregatedCollisions.Append(ExtractedCollisions);
 	// External collisions added from the world
@@ -1793,7 +1793,7 @@ void FClothingActorNv::UpdateAnimDrive(FClothingSimulationContextNv* InContext)
 {
 	SCOPE_CYCLE_COUNTER(STAT_NvClothUpdateAnimDrive);
 
-	const FClothPhysicalMeshData& MeshDataBase = AssetCreatedFrom->ClothLodData[CurrentLodIndex]->ClothPhysicalMeshData;
+	const FClothPhysicalMeshData& MeshDataBase = AssetCreatedFrom->LodData[CurrentLodIndex].PhysicalMeshData;
 	const FPointWeightMap* const AnimDriveMultipliers = MeshDataBase.FindWeightMap(EWeightMapTargetCommon::AnimDriveMultiplier);
 	if(AnimDriveMultipliers && AnimDriveMultipliers->Num())
 	{
