Cloth - Fixed broken binding in Editor on 4.24 cloth assets.

diff --git a/Engine/Source/Runtime/ClothingSystemRuntimeCommon/Private/ClothingAsset.cpp b/Engine/Source/Runtime/ClothingSystemRuntimeCommon/Private/ClothingAsset.cpp
--- a/Engine/Source/Runtime/ClothingSystemRuntimeCommon/Private/ClothingAsset.cpp
+++ b/Engine/Source/Runtime/ClothingSystemRuntimeCommon/Private/ClothingAsset.cpp
@@ -200,6 +200,19 @@ bool UClothingAssetCommon::BindToSkeletalMesh(
 	const int32 InSectionIndex, 
 	const int32 InAssetLodIndex)
 {
+	// Make sure the legacy LOD are upgraded (BindToSkeletalMesh could be called before the Cloth Asset's PostLoad is completed)
+	for (UClothLODDataCommon_Legacy* LodDeprecated : ClothLodData_DEPRECATED)
+	{
+		if (LodDeprecated)
+		{
+			LodDeprecated->ConditionalPostLoad();
+
+			const int32 Idx = AddNewLod();
+			LodDeprecated->MigrateTo(LodData[Idx]);
+		}
+	}
+	ClothLodData_DEPRECATED.Empty();
+
 	// If we've been added to the wrong mesh
 	if(InSkelMesh != GetOuter())
 	{
