Cloth - Fix cloth vertex mapping to render section when one or two points of an element are fixed, allowing to now draw triangles that are partially attached as simulating rather than skinned.

diff --git a/Engine/Source/Editor/UnrealEd/Private/Animation/DebugSkelMeshComponent.cpp b/Engine/Source/Editor/UnrealEd/Private/Animation/DebugSkelMeshComponent.cpp
--- a/Engine/Source/Editor/UnrealEd/Private/Animation/DebugSkelMeshComponent.cpp
+++ b/Engine/Source/Editor/UnrealEd/Private/Animation/DebugSkelMeshComponent.cpp
@@ -747,17 +747,10 @@ void UDebugSkelMeshComponent::RebuildClothingSectionsFixedVerts()
 
 					for(FMeshToMeshVertData& VertData : Section.ClothMappingData)
 					{
-						if(MaxDistances.AreAnyBelowThreshold(
+						VertData.SourceMeshVertIndices[3] = MaxDistances.AreAllBelowThreshold(
 							VertData.SourceMeshVertIndices[0],
 							VertData.SourceMeshVertIndices[1],
-							VertData.SourceMeshVertIndices[2]))
-						{
-							VertData.SourceMeshVertIndices[3] = 0xFFFF;
-						}
-						else
-						{
-							VertData.SourceMeshVertIndices[3] = 0;
-						}
+							VertData.SourceMeshVertIndices[2]) ? 0xFFFF : 0;
 					}
 					//We must dirty the DDC key
 					SkeletalMesh->InvalidateDeriveDataCacheGUID();
diff --git a/Engine/Source/Runtime/ClothingSystemRuntimeCommon/Private/ClothingAsset.cpp b/Engine/Source/Runtime/ClothingSystemRuntimeCommon/Private/ClothingAsset.cpp
--- a/Engine/Source/Runtime/ClothingSystemRuntimeCommon/Private/ClothingAsset.cpp
+++ b/Engine/Source/Runtime/ClothingSystemRuntimeCommon/Private/ClothingAsset.cpp
@@ -328,13 +328,10 @@ bool UClothingAssetCommon::BindToSkeletalMesh(
 	const FPointWeightMap& MaxDistances = ClothLodData.PhysicalMeshData.GetWeightMap(EWeightMapTargetCommon::MaxDistance);
 	for(FMeshToMeshVertData& VertData : MeshToMeshData)
 	{
-		if(MaxDistances.AreAnyBelowThreshold(
-			VertData.SourceMeshVertIndices[0], 
-			VertData.SourceMeshVertIndices[1], 
-			VertData.SourceMeshVertIndices[2])) // Default threshold is 0.1, not 0.0.  Using 0.1 for consistency.
-		{
-			VertData.SourceMeshVertIndices[3] = 0xFFFF;
-		}
+		VertData.SourceMeshVertIndices[3] = MaxDistances.AreAllBelowThreshold(
+			VertData.SourceMeshVertIndices[0],
+			VertData.SourceMeshVertIndices[1],
+			VertData.SourceMeshVertIndices[2]) ? 0xFFFF : 0;
 	}
 
 	// We have to copy the bone map to verify we don't exceed the maximum while adding the clothing bones
diff --git a/Engine/Source/Runtime/ClothingSystemRuntimeCommon/Public/PointWeightMap.h b/Engine/Source/Runtime/ClothingSystemRuntimeCommon/Public/PointWeightMap.h
--- a/Engine/Source/Runtime/ClothingSystemRuntimeCommon/Public/PointWeightMap.h
+++ b/Engine/Source/Runtime/ClothingSystemRuntimeCommon/Public/PointWeightMap.h
@@ -108,10 +108,14 @@ struct CLOTHINGSYSTEMRUNTIMECOMMON_API FPointWeightMap
 	bool IsBelowThreshold(const int32 Index, const float Threshold=0.1f) const
 	{ return Values.IsValidIndex(Index) && Values[Index] <= Threshold; }
 
-	/** Return whether one of the specified triangle has at least one point weight below (or equal) to the specified @param Threshold. */
+	/** Return whether at least one of the specified triangle points has weight below (or equal) to the specified @param Threshold. */
 	bool AreAnyBelowThreshold(const int32 Index0, const int32 Index1, const int32 Index2, const float Threshold=0.1f) const
 	{ return IsBelowThreshold(Index0, Threshold) || IsBelowThreshold(Index1, Threshold) || IsBelowThreshold(Index2, Threshold); }
 
+	/** Return whether all of the specified triangle points have weight below (or equal) to the specified @param Threshold. */
+	bool AreAllBelowThreshold(const int32 Index0, const int32 Index1, const int32 Index2, const float Threshold=0.1f) const
+	{ return IsBelowThreshold(Index0, Threshold) && IsBelowThreshold(Index1, Threshold) && IsBelowThreshold(Index2, Threshold); }
+
 	/** Return whether all points' values are zero. */
 	bool IsZeroed() const
 	{ return !Values.FindByPredicate([](const float& Value) { return Value != 0.f; }); }
