Cloth - Fix assert "Attempting to replace an object that hasn't been fully loaded: ClothConfigNv object" when loading 4.24 skeletal mesh assets with cloths.

diff --git a/Engine/Source/Runtime/ClothingSystemRuntimeCommon/Private/ClothingAsset.cpp b/Engine/Source/Runtime/ClothingSystemRuntimeCommon/Private/ClothingAsset.cpp
--- a/Engine/Source/Runtime/ClothingSystemRuntimeCommon/Private/ClothingAsset.cpp
+++ b/Engine/Source/Runtime/ClothingSystemRuntimeCommon/Private/ClothingAsset.cpp
@@ -819,6 +819,18 @@ void UClothingAssetCommon::PostLoad()
 			ClothConfig.Value->ConditionalPostLoad();  // PostLoad configs before adding new ones
 		}
 	}
+	if (ClothSimConfig_DEPRECATED)
+	{
+		ClothSimConfig_DEPRECATED->ConditionalPostLoad();  // PostLoad old configs before replacing them
+	}
+	if (ChaosClothSimConfig_DEPRECATED)
+	{
+		ChaosClothSimConfig_DEPRECATED->ConditionalPostLoad();  // PostLoad old configs before replacing them
+	}
+	if (ClothSharedSimConfig_DEPRECATED)
+	{
+		ClothSharedSimConfig_DEPRECATED->ConditionalPostLoad();  // PostLoad old configs before replacing them
+	}
 	AddClothConfigs();
 
 	// Migrate configs
