Accelerate cloth binding.

diff --git a/Engine/Source/Runtime/ClothingSystemRuntimeCommon/Private/Utils/ClothingMeshUtils.cpp b/Engine/Source/Runtime/ClothingSystemRuntimeCommon/Private/Utils/ClothingMeshUtils.cpp
--- a/Engine/Source/Runtime/ClothingSystemRuntimeCommon/Private/Utils/ClothingMeshUtils.cpp
+++ b/Engine/Source/Runtime/ClothingSystemRuntimeCommon/Private/Utils/ClothingMeshUtils.cpp
@@ -178,10 +178,15 @@ namespace ClothingMeshUtils
 		float MinimumDistanceSq = MAX_flt;
 		int32 ClosestBaseIndex = INDEX_NONE;
 
-		const int32 NumTriangles = Mesh.Indices.Num() / 3;
-		for(int32 TriIdx = 0; TriIdx < NumTriangles; ++TriIdx)
+		const TArray<int32> Tris = const_cast<ClothMeshDesc&>(Mesh).FindCandidateTriangles(Position);
+		int32 NumTriangles = Tris.Num();
+		if (!NumTriangles)
 		{
-			int32 TriBaseIdx = TriIdx * 3;
+			NumTriangles = Mesh.Indices.Num() / 3;
+		}
+		for (int32 TriIdx = 0; TriIdx < NumTriangles; ++TriIdx)
+		{
+			int32 TriBaseIdx = (Tris.Num() ? Tris[TriIdx] : TriIdx) * 3;
 
 			const uint32 IA = Mesh.Indices[TriBaseIdx + 0];
 			const uint32 IB = Mesh.Indices[TriBaseIdx + 1];
@@ -194,7 +199,7 @@ namespace ClothingMeshUtils
 			FVector PointOnTri = FMath::ClosestPointOnTriangleToPoint(Position, A, B, C);
 			float DistSq = (PointOnTri - Position).SizeSquared();
 
-			if(DistSq < MinimumDistanceSq)
+			if (DistSq < MinimumDistanceSq)
 			{
 				MinimumDistanceSq = DistSq;
 				ClosestBaseIndex = TriBaseIdx;
diff --git a/Engine/Source/Runtime/ClothingSystemRuntimeCommon/Public/Utils/ClothingMeshUtils.h b/Engine/Source/Runtime/ClothingSystemRuntimeCommon/Public/Utils/ClothingMeshUtils.h
--- a/Engine/Source/Runtime/ClothingSystemRuntimeCommon/Public/Utils/ClothingMeshUtils.h
+++ b/Engine/Source/Runtime/ClothingSystemRuntimeCommon/Public/Utils/ClothingMeshUtils.h
@@ -3,6 +3,8 @@
 #pragma once
 
 #include "SkeletalMeshTypes.h"
+#include "Chaos/AABBTree.h"
+#include "Chaos/GeometryParticlesfwd.h"
 
 DECLARE_LOG_CATEGORY_EXTERN(LogClothingMeshUtils, Log, All);
 
@@ -12,20 +14,92 @@ namespace ClothingMeshUtils
 {
 	struct CLOTHINGSYSTEMRUNTIMECOMMON_API ClothMeshDesc
 	{
-		ClothMeshDesc(TArrayView<const FVector> InPositions, TArrayView<const FVector> InNormals,  TArrayView<const uint32> InIndices)
+		struct FClothBvEntry
+		{
+			ClothMeshDesc* TmData;
+			int32 Index;
+
+			bool HasBoundingBox() const { return true; }
+
+			Chaos::TAABB<float, 3> BoundingBox() const
+			{
+				int32 TriBaseIdx = Index * 3;
+
+				const uint32 IA = TmData->Indices[TriBaseIdx + 0];
+				const uint32 IB = TmData->Indices[TriBaseIdx + 1];
+				const uint32 IC = TmData->Indices[TriBaseIdx + 2];
+
+				const FVector& A = TmData->Positions[IA];
+				const FVector& B = TmData->Positions[IB];
+				const FVector& C = TmData->Positions[IC];
+
+				Chaos::TAABB<float,3> Bounds(A, A);
+
+				Bounds.GrowToInclude(B);
+				Bounds.GrowToInclude(C);
+
+				return Bounds;
+			}
+
+			template<typename TPayloadType>
+			int32 GetPayload(int32 Idx) const
+			{
+				return Idx;
+			}
+
+			Chaos::FUniqueIdx UniqueIdx() const
+			{
+				return Chaos::FUniqueIdx(Index);
+			}
+		};
+
+		ClothMeshDesc(TArrayView<const FVector> InPositions, TArrayView<const FVector> InNormals, TArrayView<const uint32> InIndices)
 			: Positions(InPositions)
 			, Normals(InNormals)
 			, Indices(InIndices)
-		{}
+			, bHasValidBVH(false)
+		{
+		}
 
 		bool HasValidMesh() const
 		{
 			return Positions.Num() == Normals.Num() && Indices.Num() % 3 == 0;
 		}
 
+		TArray<int32> FindCandidateTriangles(const FVector Point)
+		{
+			ensure(HasValidMesh());
+			const int32 NumTris = Indices.Num() / 3;
+			if (NumTris > 100)
+			{
+				// This is not thread safe
+				if (!bHasValidBVH)
+				{
+					TArray<FClothBvEntry> BVEntries;
+					BVEntries.Reset(NumTris);
+
+					for (int Tri = 0; Tri < NumTris; Tri++)
+					{
+						BVEntries.Add({ this, Tri });
+					}
+					BVH.Reinitialize(BVEntries);
+					bHasValidBVH = true;
+				}
+				Chaos::TAABB<float, 3> TmpAABB(Point, Point);
+				return BVH.FindAllIntersections(TmpAABB);
+			}
+			else
+			{
+				return TArray<int32>();
+			}
+		}
+
 		TArrayView<const FVector> Positions;
 		TArrayView<const FVector> Normals;
 		TArrayView<const uint32> Indices;
+
+		bool bHasValidBVH;
+		Chaos::TAABBTree<int32, Chaos::TAABBTreeLeafArray<int32, float, false>, float, false> BVH;
 	};
 
 	// Static method for calculating a skinned mesh result from source data
