Prototype Wrapped Object Pointers (ObjectHandle/ObjectPtr)

Call-site changes for non-scalar wrapped object pointer upgrades in Engine + Shootergame

diff --git a/Engine/Plugins/Editor/DataValidation/Source/DataValidation/Private/EditorValidatorSubsystem.cpp b/Engine/Plugins/Editor/DataValidation/Source/DataValidation/Private/EditorValidatorSubsystem.cpp
--- a/Engine/Plugins/Editor/DataValidation/Source/DataValidation/Private/EditorValidatorSubsystem.cpp
+++ b/Engine/Plugins/Editor/DataValidation/Source/DataValidation/Private/EditorValidatorSubsystem.cpp
@@ -148,7 +148,7 @@ EDataValidationResult UEditorValidatorSubsystem::IsObjectValid(UObject* InObject
 		// If the asset is still valid or there wasn't a class-level validation, keep validating with custom validators
 		if (Result != EDataValidationResult::Invalid)
 		{
-			for (TPair<UClass*, UEditorValidatorBase*> ValidatorPair : Validators)
+			for (auto ValidatorPair : Validators)
 			{
 				if (ValidatorPair.Value && ValidatorPair.Value->IsEnabled() && ValidatorPair.Value->CanValidateAsset(InObject))
 				{
diff --git a/Engine/Source/Runtime/ClothingSystemRuntimeCommon/Private/ClothingAsset.cpp b/Engine/Source/Runtime/ClothingSystemRuntimeCommon/Private/ClothingAsset.cpp
--- a/Engine/Source/Runtime/ClothingSystemRuntimeCommon/Private/ClothingAsset.cpp
+++ b/Engine/Source/Runtime/ClothingSystemRuntimeCommon/Private/ClothingAsset.cpp
@@ -995,7 +995,7 @@ void UClothingAssetCommon::AddClothConfigs()
 					// migration from compatible config sources
 					if (UClothConfigCommon* const ClothConfigCommon = Cast<UClothConfigCommon>(ClothConfig))
 					{
-						for (TPair<FName, UClothConfigBase*> ClothConfigPair : ClothConfigs)
+						for (auto ClothConfigPair : ClothConfigs)
 						{
 							if (const UClothConfigCommon* SourceConfig = Cast<UClothConfigCommon>(ClothConfigPair.Value))
 							{
@@ -1026,7 +1026,7 @@ void UClothingAssetCommon::PropagateSharedConfigs(bool bMigrateSharedConfigToCon
 		const TArray<UClothingAssetBase*>& ClothingAssets = SkeletalMesh->GetMeshClothingAssets();
  
 		// Collect all shared configs found in the other assets
-		TMap<FName, UClothConfigBase*> ClothSharedConfigs;
+		decltype(ClothConfigs) ClothSharedConfigs;
 
 		for (const UClothingAssetBase* ClothingAssetBase : ClothingAssets)
 		{
@@ -1055,7 +1055,7 @@ void UClothingAssetCommon::PropagateSharedConfigs(bool bMigrateSharedConfigToCon
 		}
 
 		// Propagate the found shared configs to this asset
-		for (const TPair<FName, UClothConfigBase*>& ClothSharedConfigItem : ClothSharedConfigs)
+		for (const auto& ClothSharedConfigItem : ClothSharedConfigs)
 		{
 			// Set share config
 			if (UE_TRANSITIONAL_OBJECT_PTR(UClothConfigBase)* const ClothConfigBase = ClothConfigs.Find(ClothSharedConfigItem.Key))
diff --git a/Engine/Source/Runtime/Interchange/Core/Private/Nodes/InterchangeBaseNodeContainer.cpp b/Engine/Source/Runtime/Interchange/Core/Private/Nodes/InterchangeBaseNodeContainer.cpp
--- a/Engine/Source/Runtime/Interchange/Core/Private/Nodes/InterchangeBaseNodeContainer.cpp
+++ b/Engine/Source/Runtime/Interchange/Core/Private/Nodes/InterchangeBaseNodeContainer.cpp
@@ -164,7 +164,7 @@ void UInterchangeBaseNodeContainer::SerializeNodeContainerData(FArchive& Ar)
 			BaseNode->Serialize(Ar);
 		};
 
-		for (TPair<FString, UInterchangeBaseNode*> NodePair : Nodes)
+		for (auto NodePair : Nodes)
 		{
 			SerializeNodePair(NodePair.Value);
 		}
diff --git a/Engine/Source/Runtime/MediaAssets/Private/Assets/PlatformMediaSource.cpp b/Engine/Source/Runtime/MediaAssets/Private/Assets/PlatformMediaSource.cpp
--- a/Engine/Source/Runtime/MediaAssets/Private/Assets/PlatformMediaSource.cpp
+++ b/Engine/Source/Runtime/MediaAssets/Private/Assets/PlatformMediaSource.cpp
@@ -76,7 +76,7 @@ void UPlatformMediaSource::Serialize(FArchive& Ar)
 			if (Ar.IsLoading() && MediaCustomVersion < FMediaFrameworkObjectVersion::SerializeGUIDsInPlatformMediaSourceInsteadOfPlainNames)
 			{
 				// Load old data version
-				TMap<FString, UMediaSource*> OldPlatformMediaSources;
+				decltype(PlatformMediaSources) OldPlatformMediaSources;
 
 				Ar << OldPlatformMediaSources;
 
