Cloth - Fix the cloth wrap deformer capture optimization code in cases where the skinned points' are outside their closest cloth triangle bounding volumes.

diff --git a/Engine/Source/Runtime/ClothingSystemRuntimeCommon/Private/Utils/ClothingMeshUtils.cpp b/Engine/Source/Runtime/ClothingSystemRuntimeCommon/Private/Utils/ClothingMeshUtils.cpp
--- a/Engine/Source/Runtime/ClothingSystemRuntimeCommon/Private/Utils/ClothingMeshUtils.cpp
+++ b/Engine/Source/Runtime/ClothingSystemRuntimeCommon/Private/Utils/ClothingMeshUtils.cpp
@@ -588,6 +588,46 @@ namespace ClothingMeshUtils
 		});
 	}
 
+
+	TArray<int32> ClothMeshDesc::FindCandidateTriangles(const FVector Point)
+	{
+		ensure(HasValidMesh());
+		static const int32 MinNumTrianglesForBVHCreation = 100;
+		const int32 NumTris = Indices.Num() / 3;
+		if (NumTris > MinNumTrianglesForBVHCreation)
+		{
+			// This is not thread safe
+			if (!bHasValidBVH)
+			{
+				TArray<FClothBvEntry> BVEntries;
+				BVEntries.Reset(NumTris);
+
+				for (int32 Tri = 0; Tri < NumTris; ++Tri)
+				{
+					BVEntries.Add({ this, Tri });
+				}
+				BVH.Reinitialize(BVEntries);
+				bHasValidBVH = true;
+			}
+			Chaos::FAABB3 TmpAABB(Point, Point);
+			TmpAABB.Thicken(KINDA_SMALL_NUMBER);  // Most points might be very close to the triangle, but not directly on it
+			TArray<int32> Triangles = BVH.FindAllIntersections(TmpAABB);
+
+			// Refine the search to include all nearby bounded volumes (the point could well be outside the closest triangle's bounded volume)
+			if (Triangles.Num())
+			{
+				float ClosestDistance = TNumericLimits<float>::Max();
+				for (const int32 Triangle : Triangles)
+				{
+					ClosestDistance = FMath::Min(ClosestDistance, DistanceToTriangle(Point, *this, Triangle * 3));
+				}
+
+				TmpAABB.Thicken(ClosestDistance);
+				return BVH.FindAllIntersections(TmpAABB);
+			}
+		}
+		return TArray<int32>();
+	}
 }
 
 #undef LOCTEXT_NAMESPACE
diff --git a/Engine/Source/Runtime/ClothingSystemRuntimeCommon/Public/Utils/ClothingMeshUtils.h b/Engine/Source/Runtime/ClothingSystemRuntimeCommon/Public/Utils/ClothingMeshUtils.h
--- a/Engine/Source/Runtime/ClothingSystemRuntimeCommon/Public/Utils/ClothingMeshUtils.h
+++ b/Engine/Source/Runtime/ClothingSystemRuntimeCommon/Public/Utils/ClothingMeshUtils.h
@@ -66,33 +66,7 @@ namespace ClothingMeshUtils
 			return Positions.Num() == Normals.Num() && Indices.Num() % 3 == 0;
 		}
 
-		TArray<int32> FindCandidateTriangles(const FVector Point)
-		{
-			ensure(HasValidMesh());
-			const int32 NumTris = Indices.Num() / 3;
-			if (NumTris > 100)
-			{
-				// This is not thread safe
-				if (!bHasValidBVH)
-				{
-					TArray<FClothBvEntry> BVEntries;
-					BVEntries.Reset(NumTris);
-
-					for (int Tri = 0; Tri < NumTris; Tri++)
-					{
-						BVEntries.Add({ this, Tri });
-					}
-					BVH.Reinitialize(BVEntries);
-					bHasValidBVH = true;
-				}
-				Chaos::FAABB3 TmpAABB(Point, Point);
-				return BVH.FindAllIntersections(TmpAABB);
-			}
-			else
-			{
-				return TArray<int32>();
-			}
-		}
+		TArray<int32> FindCandidateTriangles(const FVector Point);
 
 		TArrayView<const FVector> Positions;
 		TArrayView<const FVector> Normals;
