Cloth - Fix LOD transition maps, broken by mistakenly enabling multiple influences in the call to GenerateMeshToMeshSkinningData. [Integrate CL 16114138 from //UE5/Release-Engine-Staging]

diff --git a/Engine/Source/Runtime/ClothingSystemRuntimeCommon/Private/ClothingAsset.cpp b/Engine/Source/Runtime/ClothingSystemRuntimeCommon/Private/ClothingAsset.cpp
--- a/Engine/Source/Runtime/ClothingSystemRuntimeCommon/Private/ClothingAsset.cpp
+++ b/Engine/Source/Runtime/ClothingSystemRuntimeCommon/Private/ClothingAsset.cpp
@@ -24,6 +24,7 @@
 #include "ClothingSimulationInteractor.h"
 #include "ComponentReregisterContext.h"
 #include "UObject/UObjectIterator.h"
+#include "UObject/PhysicsObjectVersion.h"
 
 #include "GPUSkinPublicDefs.h"
 #include "GPUSkinVertexFactory.h"
@@ -601,14 +602,14 @@ void UClothingAssetCommon::BuildLodTransitionData()
 		const int32 CurrentLodNumVerts = CurrentPhysMesh.Vertices.Num();
 
 		ClothingMeshUtils::ClothMeshDesc CurrentMeshDesc(CurrentPhysMesh.Vertices, CurrentPhysMesh.Normals, CurrentPhysMesh.Indices);
-
+		static const bool bUseMultipleInfluences = false;  // Multiple influences must not be used for LOD transitions
 		if(PrevLod)
 		{
 			FClothPhysicalMeshData& PrevPhysMesh = PrevLod->PhysicalMeshData;
 			CurrentLod.TransitionUpSkinData.Empty(CurrentLodNumVerts);
 			ClothingMeshUtils::ClothMeshDesc PrevMeshDesc(PrevPhysMesh.Vertices, PrevPhysMesh.Normals, PrevPhysMesh.Indices);
 			ClothingMeshUtils::GenerateMeshToMeshSkinningData(CurrentLod.TransitionUpSkinData, CurrentMeshDesc, nullptr, 
-				PrevMeshDesc, CurrentLod.bUseMultipleInfluences, CurrentLod.SkinningKernelRadius);
+				PrevMeshDesc, bUseMultipleInfluences, CurrentLod.SkinningKernelRadius);
 		}
 		if(NextLod)
 		{
@@ -616,7 +617,7 @@ void UClothingAssetCommon::BuildLodTransitionData()
 			CurrentLod.TransitionDownSkinData.Empty(CurrentLodNumVerts);
 			ClothingMeshUtils::ClothMeshDesc NextMeshDesc(NextPhysMesh.Vertices, NextPhysMesh.Normals, NextPhysMesh.Indices);
 			ClothingMeshUtils::GenerateMeshToMeshSkinningData(CurrentLod.TransitionDownSkinData, CurrentMeshDesc, 
-				nullptr, NextMeshDesc, CurrentLod.bUseMultipleInfluences, CurrentLod.SkinningKernelRadius);
+				nullptr, NextMeshDesc, bUseMultipleInfluences, CurrentLod.SkinningKernelRadius);
 		}
 	}
 }
@@ -968,6 +969,12 @@ void UClothingAssetCommon::PostLoad()
 	BuildSelfCollisionData();
 #if WITH_EDITORONLY_DATA
 	CalculateReferenceBoneIndex();
+
+	const int32 PhysicsObjectVersion = GetLinkerCustomVersion(FPhysicsObjectVersion::GUID);
+	if (PhysicsObjectVersion < FPhysicsObjectVersion::ChaosClothFixLODTransitionMaps)
+	{
+		BuildLodTransitionData();
+	}
 #endif
 }
 
@@ -976,6 +983,7 @@ void UClothingAssetCommon::Serialize(FArchive& Ar)
 	Super::Serialize(Ar);
 	Ar.UsingCustomVersion(FAnimPhysObjectVersion::GUID);
 	Ar.UsingCustomVersion(FClothingAssetCustomVersion::GUID);
+	Ar.UsingCustomVersion(FPhysicsObjectVersion::GUID);
 }
 
 void UClothingAssetCommon::AddClothConfigs()
diff --git a/Engine/Source/Runtime/Core/Public/UObject/PhysicsObjectVersion.h b/Engine/Source/Runtime/Core/Public/UObject/PhysicsObjectVersion.h
--- a/Engine/Source/Runtime/Core/Public/UObject/PhysicsObjectVersion.h
+++ b/Engine/Source/Runtime/Core/Public/UObject/PhysicsObjectVersion.h
@@ -42,6 +42,9 @@ struct CORE_API FPhysicsObjectVersion
 		// Added the Long Range Attachment stiffness weight map
 		ChaosClothAddTetherStiffnessWeightMap,
 
+		// Fix corrupted LOD transition maps
+		ChaosClothFixLODTransitionMaps,
+
 		// -----<new versions can be added above this line>-------------------------------------------------
 		VersionPlusOne,
 		LatestVersion = VersionPlusOne - 1
