Zoinks! Turns out PVS was right after all. We never end up with zero roots, thanks to [REDACTED_PERSON]!

#trivial
#rnx

diff --git a/Engine/Source/Runtime/ClothingSystemRuntimeCommon/Private/Utils/ClothingMeshUtils.cpp b/Engine/Source/Runtime/ClothingSystemRuntimeCommon/Private/Utils/ClothingMeshUtils.cpp
--- a/Engine/Source/Runtime/ClothingSystemRuntimeCommon/Private/Utils/ClothingMeshUtils.cpp
+++ b/Engine/Source/Runtime/ClothingSystemRuntimeCommon/Private/Utils/ClothingMeshUtils.cpp
@@ -682,11 +682,10 @@ namespace ClothingMeshUtils
 		// Pass in inverted normals, since they get inverted at runtime (Left handed system).
 		double W[3];
 		const int32 Count = CoplanarityParam(A, B, C, NA, NB, NC, Point, W);
-
-		if (Count == 0)
-		{
-			return GetPointBaryAndDist(A, B, C, Point);
-		}
+		
+		// We should always have a solution. The only case where we can get zero roots from the
+		// cubic equation, we handle separately.
+		checkSlow(Count != 0);
 
 		FVector4 BaryAndDist;
 		double MinDistanceSq = std::numeric_limits<double>::max();
