Cloth wrap deformer with multiple influences:
- use new barycentric weight computation
- if no triangles are within kernel radius, default to single-influence mode for that vertex
- add a missing division by NUM_INFLUENCES_PER_VERTEX when computing normals and tangents in the shader

#rnx

diff --git a/Engine/Shaders/Private/GpuSkinVertexFactory.ush b/Engine/Shaders/Private/GpuSkinVertexFactory.ush
--- a/Engine/Shaders/Private/GpuSkinVertexFactory.ush
+++ b/Engine/Shaders/Private/GpuSkinVertexFactory.ush
@@ -710,6 +710,8 @@ float3x3 SkinTangents( FVertexFactoryInput Input, FVertexFactoryIntermediates In
 	        }
         }
 
+		SimulWeight /= NUM_INFLUENCES_PER_VERTEX;
+	
         if ( NumInfluences > 0 && SumWeights > 1e-4f )
         {
 		    float InvWeight = 1.0f / SumWeights; 
diff --git a/Engine/Source/Runtime/ClothingSystemRuntimeCommon/Private/Utils/ClothingMeshUtils.cpp b/Engine/Source/Runtime/ClothingSystemRuntimeCommon/Private/Utils/ClothingMeshUtils.cpp
--- a/Engine/Source/Runtime/ClothingSystemRuntimeCommon/Private/Utils/ClothingMeshUtils.cpp
+++ b/Engine/Source/Runtime/ClothingSystemRuntimeCommon/Private/Utils/ClothingMeshUtils.cpp
@@ -277,19 +277,74 @@ namespace ClothingMeshUtils
 
 		float Kernel(float Distance, float MaxDistance)
 		{
+			if (MaxDistance == 0.0f)
+			{
+				return 0.0f;
+			}
+
 			float R = FMath::Max(0.0f, FMath::Min(1.0f, Distance / MaxDistance));
 			float R2 = R * R;
 			float R4 = R2 * R2;
 			return 1.0f + 3.0f * (R4 - R2) - R4 * R2;
 		}
 
+
+		bool SingleSkinningDataForVertex(const FVector3f& VertPosition,
+										 const FVector3f& VertNormal,
+										 const FVector3f& VertTangent,
+										 const ClothMeshDesc& SourceMesh,
+										 int32 ClosestTriangleBaseIdx,
+										 FMeshToMeshVertData& OutSkinningData)
+		{
+			const FVector3f& A = SourceMesh.Positions[SourceMesh.Indices[ClosestTriangleBaseIdx]];
+			const FVector3f& B = SourceMesh.Positions[SourceMesh.Indices[ClosestTriangleBaseIdx + 1]];
+			const FVector3f& C = SourceMesh.Positions[SourceMesh.Indices[ClosestTriangleBaseIdx + 2]];
+
+			const FVector3f& NA = SourceMesh.Normals[SourceMesh.Indices[ClosestTriangleBaseIdx]];
+			const FVector3f& NB = SourceMesh.Normals[SourceMesh.Indices[ClosestTriangleBaseIdx + 1]];
+			const FVector3f& NC = SourceMesh.Normals[SourceMesh.Indices[ClosestTriangleBaseIdx + 2]];
+
+			// Before generating the skinning data we need to check for a degenerate triangle.
+			// If we find _any_ degenerate triangles we will notify and fail to generate the skinning data
+			const FVector TriNormal = FVector::CrossProduct(B - A, C - A);
+			if (TriNormal.SizeSquared() < SMALL_NUMBER)
+			{
+				// Failed, we have 2 identical vertices
+
+				// Log and toast
+				FText Error = FText::Format(LOCTEXT("DegenerateTriangleError", "Failed to generate skinning data, found conincident vertices in triangle A={0} B={1} C={2}"), FText::FromString(A.ToString()), FText::FromString(B.ToString()), FText::FromString(C.ToString()));
+
+				UE_LOG(LogClothingMeshUtils, Warning, TEXT("%s"), *Error.ToString());
+
+#if WITH_EDITOR
+				FNotificationInfo Info(Error);
+				Info.ExpireDuration = 5.0f;
+				FSlateNotificationManager::Get().AddNotification(Info);
+#endif
+				return false;
+			}
+
+			OutSkinningData.PositionBaryCoordsAndDist = GetPointBaryAndDistWithNormals(A, B, C, NA, NB, NC, VertPosition);
+			OutSkinningData.NormalBaryCoordsAndDist = GetPointBaryAndDistWithNormals(A, B, C, NA, NB, NC, VertPosition + VertNormal);
+			OutSkinningData.TangentBaryCoordsAndDist = GetPointBaryAndDistWithNormals(A, B, C, NA, NB, NC, VertPosition + VertTangent);
+			OutSkinningData.SourceMeshVertIndices[0] = SourceMesh.Indices[ClosestTriangleBaseIdx];
+			OutSkinningData.SourceMeshVertIndices[1] = SourceMesh.Indices[ClosestTriangleBaseIdx + 1];
+			OutSkinningData.SourceMeshVertIndices[2] = SourceMesh.Indices[ClosestTriangleBaseIdx + 2];
+			OutSkinningData.SourceMeshVertIndices[3] = 0;
+			OutSkinningData.Weight = 1.0f;
+
+			return true;
+		}
+
+
 		template<unsigned int NUM_INFLUENCES>
-		bool SkinningDataForVertex(TStaticArray<FMeshToMeshVertData, NUM_INFLUENCES>& SkinningData,
-			const ClothMeshDesc& TargetMesh,
-			const TArray<FVector3f>* TargetTangents,
-			const ClothMeshDesc& SourceMesh,
-			int32 VertIdx0,
-			float KernelMaxDistance)
+		bool MultipleSkinningDataForVertex(TStaticArray<FMeshToMeshVertData, NUM_INFLUENCES>& SkinningData,
+										   const ClothMeshDesc& TargetMesh,
+										   const TArray<FVector3f>* TargetTangents,
+										   const ClothMeshDesc& SourceMesh,
+										   int32 VertIdx0,
+										   float KernelMaxDistance,
+										   float MaxIncidentEdgeLength)
 		{
 			const FVector3f& VertPosition = TargetMesh.Positions[VertIdx0];
 			const FVector3f& VertNormal = TargetMesh.Normals[VertIdx0];
@@ -351,9 +406,9 @@ namespace ClothingMeshUtils
 					return false;
 				}
 
-				CurrentData.PositionBaryCoordsAndDist = GetPointBaryAndDist(A, B, C, VertPosition);
-				CurrentData.NormalBaryCoordsAndDist = GetPointBaryAndDist(A, B, C, VertPosition + VertNormal);
-				CurrentData.TangentBaryCoordsAndDist = GetPointBaryAndDist(A, B, C, VertPosition + VertTangent);
+				CurrentData.PositionBaryCoordsAndDist = GetPointBaryAndDistWithNormals(A, B, C, NA, NB, NC, VertPosition);
+				CurrentData.NormalBaryCoordsAndDist = GetPointBaryAndDistWithNormals(A, B, C, NA, NB, NC, VertPosition + VertNormal);
+				CurrentData.TangentBaryCoordsAndDist = GetPointBaryAndDistWithNormals(A, B, C, NA, NB, NC, VertPosition + VertTangent);
 				CurrentData.SourceMeshVertIndices[0] = SourceMesh.Indices[ClosestTriangleBaseIdx];
 				CurrentData.SourceMeshVertIndices[1] = SourceMesh.Indices[ClosestTriangleBaseIdx + 1];
 				CurrentData.SourceMeshVertIndices[2] = SourceMesh.Indices[ClosestTriangleBaseIdx + 2];
@@ -374,11 +429,18 @@ namespace ClothingMeshUtils
 					CurrentData.SourceMeshVertIndices[3] = 0xFFFF;
 					CurrentData.Weight = 0.0f;
 				}
-			}
 
-			for (FMeshToMeshVertData& CurrentData : SkinningData)
+				// Compute one rigid attachment
+				const int32 ClosestTriangleBaseIdx = GetBestTriangleBaseIndex(SourceMesh, VertPosition, MaxIncidentEdgeLength);
+				check(ClosestTriangleBaseIdx != INDEX_NONE);
+				SingleSkinningDataForVertex(VertPosition, VertNormal, VertTangent, SourceMesh, ClosestTriangleBaseIdx, SkinningData[0]);
+			}
+			else
 			{
-				CurrentData.Weight /= SumWeight;
+				for (FMeshToMeshVertData& CurrentData : SkinningData)
+				{
+					CurrentData.Weight /= SumWeight;
+				}
 			}
 
 			return true;
@@ -432,6 +494,29 @@ namespace ClothingMeshUtils
 			return;
 		}
 
+
+		// Compute the max edge length for each edge coming out of a target vertex and 
+		// use that to guide the search radius for the source mesh triangles.
+		TArray<float> MaxEdgeLength;
+		MaxEdgeLength.Init(0.0f, NumMesh0Verts);
+
+		for (int32 TriangleIdx = 0; TriangleIdx < NumMesh0Verts; ++TriangleIdx)
+		{
+			const uint32* Triangle = &TargetMesh.Indices[TriangleIdx * 3];
+
+			for (int32 Vertex0Idx = 0; Vertex0Idx < 3; Vertex0Idx++)
+			{
+				const int32 Vertex1Idx = (Vertex0Idx + 1) % 3;
+
+				const FVector& P0 = TargetMesh.Positions[Triangle[Vertex0Idx]];
+				const FVector& P1 = TargetMesh.Positions[Triangle[Vertex1Idx]];
+
+				const float EdgeLength = FVector::Distance(P0, P1);
+				MaxEdgeLength[Triangle[Vertex0Idx]] = FMath::Max(MaxEdgeLength[Triangle[Vertex0Idx]], EdgeLength);
+				MaxEdgeLength[Triangle[Vertex1Idx]] = FMath::Max(MaxEdgeLength[Triangle[Vertex1Idx]], EdgeLength);
+			}
+		}
+
 		if (bUseMultipleInfluences)
 		{
 			OutSkinningData.Reserve(NumMesh0Verts * NUM_INFLUENCES_PER_VERTEX);
@@ -440,12 +525,13 @@ namespace ClothingMeshUtils
 			for (int32 VertIdx0 = 0; VertIdx0 < NumMesh0Verts; ++VertIdx0)
 			{
 				TStaticArray<FMeshToMeshVertData, NUM_INFLUENCES_PER_VERTEX> SkinningData;
-				bool bOK = SkinningDataForVertex(SkinningData,
-												 TargetMesh,
-												 TargetTangents,
-												 SourceMesh,
-												 VertIdx0,
-												 KernelMaxDistance);
+				bool bOK = MultipleSkinningDataForVertex(SkinningData,
+														 TargetMesh,
+														 TargetTangents,
+														 SourceMesh,
+														 VertIdx0,
+														 KernelMaxDistance,
+														 MaxEdgeLength[VertIdx0]);
 
 				// If we find _any_ degenerate triangles we will notify and fail to generate the skinning data
 				if (!bOK)
@@ -464,28 +550,6 @@ namespace ClothingMeshUtils
 		{
 			OutSkinningData.Reserve(NumMesh0Verts);
 
-			// Compute the max edge length for each edge coming out of a target vertex and 
-			// use that to guide the search radius for the source mesh triangles.
-			TArray<float> MaxEdgeLength;
-			MaxEdgeLength.Init(0.0f, NumMesh0Verts);
-
-			for (int32 TriangleIdx = 0; TriangleIdx < NumMesh0Verts; ++TriangleIdx)
-			{
-				const uint32* Triangle = &TargetMesh.Indices[TriangleIdx * 3]; 
-				
-				for (int32 Vertex0Idx = 0; Vertex0Idx < 3; Vertex0Idx++)
-				{
-					const int32 Vertex1Idx = (Vertex0Idx + 1) % 3;
-
-					const FVector& P0 = TargetMesh.Positions[Triangle[Vertex0Idx]];
-					const FVector& P1 = TargetMesh.Positions[Triangle[Vertex1Idx]];
-
-					const float EdgeLength = FVector::Distance(P0, P1);
-					MaxEdgeLength[Triangle[Vertex0Idx]] = FMath::Max(MaxEdgeLength[Triangle[Vertex0Idx]], EdgeLength);
-					MaxEdgeLength[Triangle[Vertex1Idx]] = FMath::Max(MaxEdgeLength[Triangle[Vertex1Idx]], EdgeLength);
-				}
-			}			
-
 			// For all mesh0 verts
 			for (int32 VertIdx0 = 0; VertIdx0 < NumMesh0Verts; ++VertIdx0)
 			{
@@ -510,43 +574,7 @@ namespace ClothingMeshUtils
 				const int32 ClosestTriangleBaseIdx = GetBestTriangleBaseIndex(SourceMesh, VertPosition, MaxEdgeLength[VertIdx0]);
 				check(ClosestTriangleBaseIdx != INDEX_NONE);
 
-				const FVector3f& A = SourceMesh.Positions[SourceMesh.Indices[ClosestTriangleBaseIdx]];
-				const FVector3f& B = SourceMesh.Positions[SourceMesh.Indices[ClosestTriangleBaseIdx + 1]];
-				const FVector3f& C = SourceMesh.Positions[SourceMesh.Indices[ClosestTriangleBaseIdx + 2]];
-
-				const FVector3f& NA = SourceMesh.Normals[SourceMesh.Indices[ClosestTriangleBaseIdx]];
-				const FVector3f& NB = SourceMesh.Normals[SourceMesh.Indices[ClosestTriangleBaseIdx + 1]];
-				const FVector3f& NC = SourceMesh.Normals[SourceMesh.Indices[ClosestTriangleBaseIdx + 2]];
-
-				// Before generating the skinning data we need to check for a degenerate triangle.
-				// If we find _any_ degenerate triangles we will notify and fail to generate the skinning data
-				const FVector TriNormal = FVector::CrossProduct(B - A, C - A);
-				if (TriNormal.SizeSquared() < SMALL_NUMBER)
-				{
-					// Failed, we have 2 identical vertices
-					OutSkinningData.Reset();
-
-					// Log and toast
-					FText Error = FText::Format(LOCTEXT("DegenerateTriangleError", "Failed to generate skinning data, found conincident vertices in triangle A={0} B={1} C={2}"), FText::FromString(A.ToString()), FText::FromString(B.ToString()), FText::FromString(C.ToString()));
-
-					UE_LOG(LogClothingMeshUtils, Warning, TEXT("%s"), *Error.ToString());
-
-#if WITH_EDITOR
-					FNotificationInfo Info(Error);
-					Info.ExpireDuration = 5.0f;
-					FSlateNotificationManager::Get().AddNotification(Info);
-#endif
-					return;
-				}
-
-				SkinningData.PositionBaryCoordsAndDist = GetPointBaryAndDistWithNormals(A, B, C, NA, NB, NC, VertPosition);
-				SkinningData.NormalBaryCoordsAndDist = GetPointBaryAndDistWithNormals(A, B, C, NA, NB, NC, VertPosition + VertNormal);
-				SkinningData.TangentBaryCoordsAndDist = GetPointBaryAndDistWithNormals(A, B, C, NA, NB, NC, VertPosition + VertTangent);
-				SkinningData.SourceMeshVertIndices[0] = SourceMesh.Indices[ClosestTriangleBaseIdx];
-				SkinningData.SourceMeshVertIndices[1] = SourceMesh.Indices[ClosestTriangleBaseIdx + 1];
-				SkinningData.SourceMeshVertIndices[2] = SourceMesh.Indices[ClosestTriangleBaseIdx + 2];
-				SkinningData.SourceMeshVertIndices[3] = 0;
-				SkinningData.Weight = 1.0f;
+				SingleSkinningDataForVertex(VertPosition, VertNormal, VertTangent, SourceMesh, ClosestTriangleBaseIdx, SkinningData);
 			}
 
 			check(OutSkinningData.Num() == NumMesh0Verts);
