ClothingMeshUtils:
- fix zero-weight vertices after applying max distance falloff function
- if a zero-weight vertex is ever encountered in the shader, fall back to skinned vertex data
- fix construction of TargetMesh index buffer
- pull out MaxEdgeLength computation into its own function
- simplify Kernel expression

#rnx

diff --git a/Engine/Shaders/Private/GpuSkinVertexFactory.ush b/Engine/Shaders/Private/GpuSkinVertexFactory.ush
--- a/Engine/Shaders/Private/GpuSkinVertexFactory.ush
+++ b/Engine/Shaders/Private/GpuSkinVertexFactory.ush
@@ -413,7 +413,6 @@ float4 ClothingPosition(FClothVertex ClothVertexInfluences[NUM_INFLUENCES_PER_VE
 			float Weight = 1.0f;
 			if ( NUM_INFLUENCES_PER_VERTEX > 1 )
 			{
-				// Weight is packed in the last coordinate
 				Weight = Influence.Weight;
 				SumWeights += Weight;
 			}
@@ -444,6 +443,7 @@ float4 ClothingPosition(FClothVertex ClothVertexInfluences[NUM_INFLUENCES_PER_VE
 	else
 	{
 		AveragedSimulatedPosition = float3(0, 0, 0);
+		SimulWeight = 0.0f;			// Fall back to skinned vertex position by setting this lerp param to zero
 	}
 
 	SimulWeight /= NUM_INFLUENCES_PER_VERTEX;
@@ -724,20 +724,21 @@ float3x3 SkinTangents( FVertexFactoryInput Input, FVertexFactoryIntermediates In
 		    // cloth data are all in world space so need to change into local space
 		    TangentToLocal[0] = mul(TangentToLocal[0], (half3x3)Intermediates.WorldToLocal);
 		    TangentToLocal[2] = mul(TangentToLocal[2], (half3x3)Intermediates.WorldToLocal);
+	
+			// derive the new binormal by getting the cross product of the normal and tangent
+			// and flip vector based on sign of tangent basis determinant
+			TangentToLocal[1] = cross(TangentToLocal[2], TangentToLocal[0]) * LocalTangentZ.w;
+	
+			ClothTangent = TangentToLocal;
         }
 	    else
 	    {
-		    TempClothBlendWeight = 0.0f;
+		    SimulWeight = 0.0f;			// Fall back to skinned data by setting this lerp param to zero
 	    }
 
 		TempClothBlendWeight *= SimulWeight;
-
-		// derive the new binormal by getting the cross product of the normal and tangent
-		// and flip vector based on sign of tangent basis determinant
-		TangentToLocal[1] = cross(TangentToLocal[2], TangentToLocal[0]) * LocalTangentZ.w;
-	
-		ClothTangent = TangentToLocal;
 	}
+	
 	#endif // GPUSKIN_APEX_CLOTH
 
 #endif // GPUSKIN_MORPH_BLEND
diff --git a/Engine/Source/Runtime/ClothingSystemRuntimeCommon/Private/ClothingAsset.cpp b/Engine/Source/Runtime/ClothingSystemRuntimeCommon/Private/ClothingAsset.cpp
--- a/Engine/Source/Runtime/ClothingSystemRuntimeCommon/Private/ClothingAsset.cpp
+++ b/Engine/Source/Runtime/ClothingSystemRuntimeCommon/Private/ClothingAsset.cpp
@@ -310,15 +310,30 @@ bool UClothingAssetCommon::BindToSkeletalMesh(
 	}
 
 	TArrayView<uint32> IndexView(SkelLod.IndexBuffer);
-	IndexView.Slice(OriginalSection.BaseIndex, OriginalSection.NumTriangles * 3);
+	IndexView = IndexView.Slice(OriginalSection.BaseIndex, OriginalSection.NumTriangles * 3);
+
+	TArray<uint32> RenderIndices;
+	RenderIndices.Reserve(OriginalSection.NumTriangles * 3);
+	for (uint32 OriginalIndex : IndexView)
+	{
+		const int32 TempIndex = (int32)OriginalIndex - (int32)OriginalSection.BaseVertexIndex;
+		if (ensure(RenderIndices.IsValidIndex(TempIndex)))
+		{
+			RenderIndices.Add(TempIndex);
+		}
+	}
+
+	ClothingMeshUtils::ClothMeshDesc TargetMesh(RenderPositions, RenderNormals, RenderIndices);
 
-	ClothingMeshUtils::ClothMeshDesc TargetMesh(RenderPositions, RenderNormals, IndexView);
 	ClothingMeshUtils::ClothMeshDesc SourceMesh(
 		ClothLodData.PhysicalMeshData.Vertices, 
 		ClothLodData.PhysicalMeshData.Normals, 
 		ClothLodData.PhysicalMeshData.Indices);
 
-	ClothingMeshUtils::GenerateMeshToMeshSkinningData(MeshToMeshData, TargetMesh, &RenderTangents, SourceMesh, 
+	TArray<float> MaxEdgeLength;
+	ClothingMeshUtils::ComputeMaxEdgeLength(TargetMesh, MaxEdgeLength);
+
+	ClothingMeshUtils::GenerateMeshToMeshSkinningData(MeshToMeshData, TargetMesh, &RenderTangents, SourceMesh, MaxEdgeLength,
 		ClothLodData.bUseMultipleInfluences, ClothLodData.SkinningKernelRadius);
 
 	if(MeshToMeshData.Num() == 0)
@@ -330,9 +345,13 @@ bool UClothingAssetCommon::BindToSkeletalMesh(
 
 	// Calculate the vertex contribution alpha
 	const FPointWeightMap* const MaxDistances = ClothLodData.PhysicalMeshData.FindWeightMap(EWeightMapTargetCommon::MaxDistance);
-
 	ClothingMeshUtils::ComputeVertexContributions(MeshToMeshData, MaxDistances, ClothLodData.bSmoothTransition);
 
+	if (ClothLodData.bUseMultipleInfluences)
+	{
+		ClothingMeshUtils::FixZeroWeightVertices(MeshToMeshData, TargetMesh, &RenderTangents, SourceMesh, MaxEdgeLength);
+	}
+
 	// We have to copy the bone map to verify we don't exceed the maximum while adding the clothing bones
 	TArray<FBoneIndexType> TempBoneMap = OriginalSection.BoneMap;
 	for(FName& BoneName : UsedBoneNames)
@@ -579,13 +598,16 @@ void UClothingAssetCommon::BuildLodTransitionData()
 
 		ClothingMeshUtils::ClothMeshDesc CurrentMeshDesc(CurrentPhysMesh.Vertices, CurrentPhysMesh.Normals, CurrentPhysMesh.Indices);
 		static const bool bUseMultipleInfluences = false;  // Multiple influences must not be used for LOD transitions
+
+		TArray<float> MaxEdgeLength;
+		ClothingMeshUtils::ComputeMaxEdgeLength(CurrentMeshDesc, MaxEdgeLength);
 		if(PrevLod)
 		{
 			FClothPhysicalMeshData& PrevPhysMesh = PrevLod->PhysicalMeshData;
 			CurrentLod.TransitionUpSkinData.Empty(CurrentLodNumVerts);
 			ClothingMeshUtils::ClothMeshDesc PrevMeshDesc(PrevPhysMesh.Vertices, PrevPhysMesh.Normals, PrevPhysMesh.Indices);
 			ClothingMeshUtils::GenerateMeshToMeshSkinningData(CurrentLod.TransitionUpSkinData, CurrentMeshDesc, nullptr, 
-				PrevMeshDesc, bUseMultipleInfluences, CurrentLod.SkinningKernelRadius);
+				PrevMeshDesc, MaxEdgeLength, bUseMultipleInfluences, CurrentLod.SkinningKernelRadius);
 		}
 		if(NextLod)
 		{
@@ -593,7 +615,7 @@ void UClothingAssetCommon::BuildLodTransitionData()
 			CurrentLod.TransitionDownSkinData.Empty(CurrentLodNumVerts);
 			ClothingMeshUtils::ClothMeshDesc NextMeshDesc(NextPhysMesh.Vertices, NextPhysMesh.Normals, NextPhysMesh.Indices);
 			ClothingMeshUtils::GenerateMeshToMeshSkinningData(CurrentLod.TransitionDownSkinData, CurrentMeshDesc, 
-				nullptr, NextMeshDesc, bUseMultipleInfluences, CurrentLod.SkinningKernelRadius);
+				nullptr, NextMeshDesc, MaxEdgeLength, bUseMultipleInfluences, CurrentLod.SkinningKernelRadius);
 		}
 	}
 }
diff --git a/Engine/Source/Runtime/ClothingSystemRuntimeCommon/Private/Utils/ClothingMeshUtils.cpp b/Engine/Source/Runtime/ClothingSystemRuntimeCommon/Private/Utils/ClothingMeshUtils.cpp
--- a/Engine/Source/Runtime/ClothingSystemRuntimeCommon/Private/Utils/ClothingMeshUtils.cpp
+++ b/Engine/Source/Runtime/ClothingSystemRuntimeCommon/Private/Utils/ClothingMeshUtils.cpp
@@ -271,7 +271,7 @@ namespace ClothingMeshUtils
 		}
 
 		// Using this formula, for R = Distance / MaxDistance:
-		//		Weight = 1 - 3 * R ^ 2 + 3 * R ^ 4 - R ^ 6
+		//		Weight = 1 - 3 * R ^ 2 + 3 * R ^ 4 - R ^ 6 = (1-R^2)^3
 		// From the Houdini metaballs docs: https://www.sidefx.com/docs/houdini/nodes/sop/metaball.html#kernels
 		// Which was linked to from the cloth capture doc: https://www.sidefx.com/docs/houdini/nodes/sop/clothcapture.html
 
@@ -283,9 +283,8 @@ namespace ClothingMeshUtils
 			}
 
 			float R = FMath::Max(0.0f, FMath::Min(1.0f, Distance / MaxDistance));
-			float R2 = R * R;
-			float R4 = R2 * R2;
-			return 1.0f + 3.0f * (R4 - R2) - R4 * R2;
+			float OneMinusR2 = 1.0f - R * R;
+			return OneMinusR2 * OneMinusR2 * OneMinusR2;
 		}
 
 
@@ -419,23 +418,7 @@ namespace ClothingMeshUtils
 			}
 
 			// Normalize weights
-
-			if (SumWeight == 0.0f)
-			{
-				// Abort
-				for (int j = 0; j < NUM_INFLUENCES; ++j)
-				{
-					FMeshToMeshVertData& CurrentData = SkinningData[j];
-					CurrentData.SourceMeshVertIndices[3] = 0xFFFF;
-					CurrentData.Weight = 0.0f;
-				}
-
-				// Compute one rigid attachment
-				const int32 ClosestTriangleBaseIdx = GetBestTriangleBaseIndex(SourceMesh, VertPosition, MaxIncidentEdgeLength);
-				check(ClosestTriangleBaseIdx != INDEX_NONE);
-				SingleSkinningDataForVertex(VertPosition, VertNormal, VertTangent, SourceMesh, ClosestTriangleBaseIdx, SkinningData[0]);
-			}
-			else
+			if (SumWeight > 0.0f)
 			{
 				for (FMeshToMeshVertData& CurrentData : SkinningData)
 				{
@@ -449,12 +432,13 @@ namespace ClothingMeshUtils
 	}		// unnamed namespace
 
 
-	void GenerateMeshToMeshSkinningData(TArray<FMeshToMeshVertData>& OutSkinningData, 
-		const ClothMeshDesc& TargetMesh, 
-		const TArray<FVector3f>* TargetTangents, 
-		const ClothMeshDesc& SourceMesh,
-		bool bUseMultipleInfluences,
-		float KernelMaxDistance)
+	void GenerateMeshToMeshSkinningData(TArray<FMeshToMeshVertData>& OutSkinningData,
+										const ClothMeshDesc& TargetMesh,
+										const TArray<FVector3f>* TargetTangents,
+										const ClothMeshDesc& SourceMesh,
+										const TArray<float>& MaxEdgeLength,
+										bool bUseMultipleInfluences,
+										float KernelMaxDistance)
 	{
 		if(!TargetMesh.HasValidMesh())
 		{
@@ -472,15 +456,14 @@ namespace ClothingMeshUtils
 		const int32 NumMesh0Normals = TargetMesh.Normals.Num();
 		const int32 NumMesh0Tangents = TargetTangents ? TargetTangents->Num() : 0;
 
+		// Check we have properly formed triangles
+		ensure(TargetMesh.Indices.Num() % 3 == 0);
+		const int32 NumMesh0Tris = TargetMesh.Indices.Num() / 3;
+
 		const int32 NumMesh1Verts = SourceMesh.Positions.Num();
 		const int32 NumMesh1Normals = SourceMesh.Normals.Num();
 		const int32 NumMesh1Indices = SourceMesh.Indices.Num();
 
-		// Check we have properly formed triangles
-		check(NumMesh1Indices % 3 == 0);
-
-		const int32 NumMesh1Triangles = NumMesh1Indices / 3;
-
 		// Check mesh data to make sure we have the same number of each element
 		if(NumMesh0Verts != NumMesh0Normals || (TargetTangents && NumMesh0Tangents != NumMesh0Verts))
 		{
@@ -494,27 +477,9 @@ namespace ClothingMeshUtils
 			return;
 		}
 
-
-		// Compute the max edge length for each edge coming out of a target vertex and 
-		// use that to guide the search radius for the source mesh triangles.
-		TArray<float> MaxEdgeLength;
-		MaxEdgeLength.Init(0.0f, NumMesh0Verts);
-
-		for (int32 TriangleIdx = 0; TriangleIdx < NumMesh0Verts; ++TriangleIdx)
+		if (!ensure(NumMesh0Verts == MaxEdgeLength.Num()))
 		{
-			const uint32* Triangle = &TargetMesh.Indices[TriangleIdx * 3];
-
-			for (int32 Vertex0Idx = 0; Vertex0Idx < 3; Vertex0Idx++)
-			{
-				const int32 Vertex1Idx = (Vertex0Idx + 1) % 3;
-
-				const FVector& P0 = TargetMesh.Positions[Triangle[Vertex0Idx]];
-				const FVector& P1 = TargetMesh.Positions[Triangle[Vertex1Idx]];
-
-				const float EdgeLength = FVector::Distance(P0, P1);
-				MaxEdgeLength[Triangle[Vertex0Idx]] = FMath::Max(MaxEdgeLength[Triangle[Vertex0Idx]], EdgeLength);
-				MaxEdgeLength[Triangle[Vertex1Idx]] = FMath::Max(MaxEdgeLength[Triangle[Vertex1Idx]], EdgeLength);
-			}
+			return;
 		}
 
 		if (bUseMultipleInfluences)
diff --git a/Engine/Source/Runtime/ClothingSystemRuntimeCommon/Public/Utils/ClothingMeshUtils.h b/Engine/Source/Runtime/ClothingSystemRuntimeCommon/Public/Utils/ClothingMeshUtils.h
--- a/Engine/Source/Runtime/ClothingSystemRuntimeCommon/Public/Utils/ClothingMeshUtils.h
+++ b/Engine/Source/Runtime/ClothingSystemRuntimeCommon/Public/Utils/ClothingMeshUtils.h
@@ -92,19 +92,57 @@ namespace ClothingMeshUtils
 		TArray<FVector3f>& OutNormals,
 		uint32 ArrayOffset = 0); // Used for Chaos Cloth
 
+
 	/**
-	* Given mesh information for two meshes, generate a list of skinning data to embed mesh0 in mesh1
-	* @param OutSkinningData	- Final skinning data to map mesh0 to mesh1
-	* @param TargetMesh		- Mesh data for the mesh we are embedding
-	* @param TargetTangents		- Optional Tangents for the mesh we are embedding
-	* @param Mesh1Normals		- Vertex normals for Mesh1
-	* @param Mesh1Indices		- Triangle indices for Mesh1
+	 * Compute the max edge length for each edge coming out of a mesh vertex. Useful for guiding the search radius when
+	 * searching for nearest triangles.
+	 */
+	inline void CLOTHINGSYSTEMRUNTIMECOMMON_API ComputeMaxEdgeLength(const ClothMeshDesc& TargetMesh, 
+																	 TArray<float>& OutMaxEdgeLength)
+	{
+		const int32 NumMesh0Verts = TargetMesh.Positions.Num();
+
+		// Check we have properly formed triangles
+		ensure(TargetMesh.Indices.Num() % 3 == 0);
+		const int32 NumMesh0Tris = TargetMesh.Indices.Num() / 3;
+
+		OutMaxEdgeLength.Init(0.0f, NumMesh0Verts);
+
+		for (int32 TriangleIdx = 0; TriangleIdx < NumMesh0Tris; ++TriangleIdx)
+		{
+			const uint32* Triangle = &TargetMesh.Indices[TriangleIdx * 3];
+
+			for (int32 Vertex0Idx = 0; Vertex0Idx < 3; ++Vertex0Idx)
+			{
+				const int32 Vertex1Idx = (Vertex0Idx + 1) % 3;
+
+				const FVector& P0 = TargetMesh.Positions[Triangle[Vertex0Idx]];
+				const FVector& P1 = TargetMesh.Positions[Triangle[Vertex1Idx]];
+
+				const float EdgeLength = FVector::Distance(P0, P1);
+				OutMaxEdgeLength[Triangle[Vertex0Idx]] = FMath::Max(OutMaxEdgeLength[Triangle[Vertex0Idx]], EdgeLength);
+				OutMaxEdgeLength[Triangle[Vertex1Idx]] = FMath::Max(OutMaxEdgeLength[Triangle[Vertex1Idx]], EdgeLength);
+			}
+		}
+	}
+
+	/**
+	* Given mesh information for two meshes, generate a list of skinning data to embed TargetMesh in SourceMesh
+	* 
+	* @param OutSkinningData            - Final skinning data
+	* @param TargetMesh                 - Mesh data for the mesh we are embedding
+	* @param TargetTangents             - Optional Tangents for the mesh we are embedding
+	* @param SourceMesh                 - Mesh data for the mesh we are embedding into
+	* @param TargetMaxEdgeLength        - Per-vertex longest incident edge length, as returned by ComputeMaxEdgeLength()
+	* @param bUseMultipleInfluences     - Whether to take a weighted average of influences from multiple source triangles
+	* @param KernelMaxDistance          - Max distance parameter for weighting kernel
 	*/
 	void CLOTHINGSYSTEMRUNTIMECOMMON_API GenerateMeshToMeshSkinningData(
 		TArray<FMeshToMeshVertData>& OutSkinningData,
 		const ClothMeshDesc& TargetMesh,
 		const TArray<FVector3f>* TargetTangents,
 		const ClothMeshDesc& SourceMesh,
+		const TArray<float>& TargetMaxEdgeLength,
 		bool bUseMultipleInfluences,
 		float KernelMaxDistance);
 
@@ -128,9 +166,19 @@ namespace ClothingMeshUtils
 	void CLOTHINGSYSTEMRUNTIMECOMMON_API ComputeVertexContributions(
 		TArray<FMeshToMeshVertData> &InOutSkinningData,
 		const FPointWeightMap* const InMaxDistances,
-		const bool bInSmoothTransition
-		);
+		const bool bInSmoothTransition );
 	
+	/**
+	 * Identify vertices that are not influenced by any triangles, and compute a new single attachment for the
+	 * vertex.
+	 */
+	void CLOTHINGSYSTEMRUNTIMECOMMON_API FixZeroWeightVertices(TArray<FMeshToMeshVertData>& InOutSkinningData,
+															   const ClothMeshDesc& TargetMesh,
+															   const TArray<FVector3f>* TargetTangents,
+															   const ClothMeshDesc& SourceMesh,
+															   const TArray<float>& TargetMaxEdgeLength );
+
+
 	/**
 	* Given a triangle ABC with normals at each vertex NA, NB and NC, get a barycentric coordinate
 	* and corresponding distance from the triangle encoded in an FVector4 where the components are
