Fix bad merge

#rnx

diff --git a/Engine/Source/Runtime/ClothingSystemRuntimeCommon/Private/Utils/ClothingMeshUtils.cpp b/Engine/Source/Runtime/ClothingSystemRuntimeCommon/Private/Utils/ClothingMeshUtils.cpp
--- a/Engine/Source/Runtime/ClothingSystemRuntimeCommon/Private/Utils/ClothingMeshUtils.cpp
+++ b/Engine/Source/Runtime/ClothingSystemRuntimeCommon/Private/Utils/ClothingMeshUtils.cpp
@@ -546,6 +546,86 @@ namespace ClothingMeshUtils
 		}
 	}
 
+
+	void FixZeroWeightVertices(TArray<FMeshToMeshVertData>& InOutSkinningData,
+							   const ClothMeshDesc& TargetMesh,
+							   const TArray<FVector3f>* TargetTangents,
+							   const ClothMeshDesc& SourceMesh,
+							   const TArray<float>& MaxEdgeLength)
+	{
+		if (!ensure(InOutSkinningData.Num() > 0))
+		{
+			return;
+		}
+
+		const int32 NumTargetMeshVerts = TargetMesh.Positions.Num();
+
+		if (!ensure(NumTargetMeshVerts * NUM_INFLUENCES_PER_VERTEX == InOutSkinningData.Num()))
+		{
+			return;
+		}
+
+		if (!ensure(NumTargetMeshVerts == MaxEdgeLength.Num()))
+		{
+			return;
+		}
+
+		for (int32 VID = 0; VID < NumTargetMeshVerts; ++VID)
+		{
+			float SumWeight = 0.0f;
+			for (int Inf = 0; Inf < NUM_INFLUENCES_PER_VERTEX; ++Inf)
+			{
+				const FMeshToMeshVertData& Data = InOutSkinningData[NUM_INFLUENCES_PER_VERTEX * VID + Inf];
+				if (Data.SourceMeshVertIndices[3] < 0xFFFF)
+				{
+					SumWeight += Data.Weight;
+				}
+			}
+
+			if (SumWeight < KINDA_SMALL_NUMBER)
+			{
+				// Fall back to single influence
+				const FVector3f& VertPosition = TargetMesh.Positions[VID];
+				const FVector3f& VertNormal = TargetMesh.Normals[VID];
+
+				FVector3f VertTangent;
+				if (TargetTangents)
+				{
+					VertTangent = (*TargetTangents)[VID];
+				}
+				else
+				{
+					FVector3f Tan0, Tan1;
+					VertNormal.FindBestAxisVectors(Tan0, Tan1);
+					VertTangent = Tan0;
+				}
+
+				// Compute single-influence attachment for the first skinning data
+				FMeshToMeshVertData& FirstData = InOutSkinningData[NUM_INFLUENCES_PER_VERTEX * VID];
+
+				const int32 ClosestTriangleBaseIdx = ClothingMeshUtils::GetBestTriangleBaseIndex(SourceMesh, VertPosition, MaxEdgeLength[VID]);
+				if (!ensure(ClosestTriangleBaseIdx != INDEX_NONE))
+				{
+					FirstData.Weight = 0.0f;
+				}
+				else
+				{
+					const uint16 PreviousFlag = FirstData.SourceMeshVertIndices[3];
+					SingleSkinningDataForVertex(VertPosition, VertNormal, VertTangent, SourceMesh, ClosestTriangleBaseIdx, FirstData);
+					FirstData.SourceMeshVertIndices[3] = PreviousFlag;
+				}
+
+				// Set all other skinning data to have zero weight
+				for (int Inf = 1; Inf < NUM_INFLUENCES_PER_VERTEX; ++Inf)
+				{
+					FMeshToMeshVertData& CurrentData = InOutSkinningData[NUM_INFLUENCES_PER_VERTEX * VID + Inf];
+					CurrentData.SourceMeshVertIndices[3] = 0xFFFF;
+					CurrentData.Weight = 0.0f;
+				}
+			}
+		}
+	}
+
 	// TODO: Vertex normals are not used at present, a future improved algorithm might however
 	FVector4 GetPointBaryAndDist(const FVector3f& A, const FVector3f& B, const FVector3f& C, const FVector3f& Point)
 	{
