Fix an ensure that will always fire

#rnx

diff --git a/Engine/Source/Runtime/ClothingSystemRuntimeCommon/Private/ClothingAsset.cpp b/Engine/Source/Runtime/ClothingSystemRuntimeCommon/Private/ClothingAsset.cpp
--- a/Engine/Source/Runtime/ClothingSystemRuntimeCommon/Private/ClothingAsset.cpp
+++ b/Engine/Source/Runtime/ClothingSystemRuntimeCommon/Private/ClothingAsset.cpp
@@ -317,7 +317,7 @@ bool UClothingAssetCommon::BindToSkeletalMesh(
 	for (uint32 OriginalIndex : IndexView)
 	{
 		const int32 TempIndex = (int32)OriginalIndex - (int32)OriginalSection.BaseVertexIndex;
-		if (ensure(RenderIndices.IsValidIndex(TempIndex)))
+		if (ensure(RenderPositions.IsValidIndex(TempIndex)))
 		{
 			RenderIndices.Add(TempIndex);
 		}
