Chaos Cloth - Significantly reduce the cloth initialization hitch by caching the Long Range Attachment tether calculations. Move the tether code to its own files inside the cloth common runtime module.
Cloth - Tidy up ClothingAsset and ClothPhysicalMeshData to allow for more cloth data caching.

diff --git a/Engine/Plugins/Experimental/ChaosCloth/Source/Chaos/Private/ChaosCloth/ChaosClothConstraints.cpp b/Engine/Plugins/Experimental/ChaosCloth/Source/Chaos/Private/ChaosCloth/ChaosClothConstraints.cpp
--- a/Engine/Plugins/Experimental/ChaosCloth/Source/Chaos/Private/ChaosCloth/ChaosClothConstraints.cpp
+++ b/Engine/Plugins/Experimental/ChaosCloth/Source/Chaos/Private/ChaosCloth/ChaosClothConstraints.cpp
@@ -390,27 +390,25 @@ void FClothConstraints::SetVolumeConstraints(TArray<TVec3<int32>>&& SurfaceEleme
 	++NumConstraintRules;
 }
 
-void FClothConstraints::SetLongRangeConstraints(const TMap<int32, TSet<int32>>& PointToNeighborsMap,
-	const TConstArrayView<FRealSingle>& TetherStiffnessMultipliers, const TConstArrayView<FRealSingle>& TetherScaleMultipliers,
-	const FVec2& TetherScale, ETetherMode TetherMode, bool bUseXPBDConstraints)
+void FClothConstraints::SetLongRangeConstraints(
+	const TArray<TConstArrayView<TTuple<int32, int32, FRealSingle>>>& Tethers,
+	const TConstArrayView<FRealSingle>& TetherStiffnessMultipliers,
+	const TConstArrayView<FRealSingle>& TetherScaleMultipliers,
+	const FVec2& TetherScale, bool bUseXPBDConstraints)
 {
 	check(Evolution);
 
-	static const int32 MaxNumTetherIslands = 4;  // The max number of connected neighbors per particle.
-
 	if (bUseXPBDConstraints)
 	{
 		XLongRangeConstraints = MakeShared<FXPBDLongRangeConstraints>(
 			Evolution->Particles(),
 			ParticleOffset,
 			NumParticles,
-			PointToNeighborsMap,
+			Tethers,
 			TetherStiffnessMultipliers,
 			TetherScaleMultipliers,
-			MaxNumTetherIslands,
 			/*InStiffness =*/ FVec2::UnitVector,
-			TetherScale,
-			TetherMode);
+			TetherScale);
 	}
 	else
 	{
@@ -418,13 +416,11 @@ void FClothConstraints::SetLongRangeConstraints(const TMap<int32, TSet<int32>>&
 			Evolution->Particles(),
 			ParticleOffset,
 			NumParticles,
-			PointToNeighborsMap,
+			Tethers,
 			TetherStiffnessMultipliers,
 			TetherScaleMultipliers,
-			MaxNumTetherIslands,
 			/*InStiffness =*/ FVec2::UnitVector,
-			TetherScale,
-			TetherMode);
+			TetherScale);
 	}
 	++NumConstraintInits;  // Uses init to update the property tables
 	++NumConstraintRules;
@@ -545,15 +541,17 @@ void FClothConstraints::SetVolumeProperties(FReal VolumeStiffness)
 	}
 }
 
-void FClothConstraints::SetLongRangeAttachmentProperties(const FVec2& TetherStiffness)
+void FClothConstraints::SetLongRangeAttachmentProperties(const FVec2& TetherStiffness, const FVec2& TetherScale)
 {
 	if (LongRangeConstraints)
 	{
 		LongRangeConstraints->SetStiffness(TetherStiffness);
+		LongRangeConstraints->SetScale(TetherScale);
 	}
 	if (XLongRangeConstraints)
 	{
 		XLongRangeConstraints->SetStiffness(TetherStiffness);
+		XLongRangeConstraints->SetScale(TetherScale);
 	}
 }
 
diff --git a/Engine/Plugins/Experimental/ChaosCloth/Source/Chaos/Private/ChaosCloth/ChaosClothingSimulation.cpp b/Engine/Plugins/Experimental/ChaosCloth/Source/Chaos/Private/ChaosCloth/ChaosClothingSimulation.cpp
--- a/Engine/Plugins/Experimental/ChaosCloth/Source/Chaos/Private/ChaosCloth/ChaosClothingSimulation.cpp
+++ b/Engine/Plugins/Experimental/ChaosCloth/Source/Chaos/Private/ChaosCloth/ChaosClothingSimulation.cpp
@@ -1740,113 +1740,45 @@ void FClothingSimulation::DebugDrawLongRangeConstraint(FPrimitiveDrawInterface*
 			continue;
 		}
 
-		// Recompute islands
-		const FTriangleMesh& TriangleMesh = Cloth->GetTriangleMesh(Solver.Get());
-		const TConstArrayView<FReal> InvMasses = Cloth->GetParticleInvMasses(Solver.Get());
-
-		const TMap<int32, TSet<int32>>& PointToNeighborsMap = TriangleMesh.GetPointToNeighborsMap();
-
-		static TArray<int32> KinematicIndices;  // Make static to prevent constant allocations
-		KinematicIndices.Reset();
-		for (const TPair<int32, TSet<int32>>& PointNeighbors : PointToNeighborsMap)
-		{
-			const int32 Index = PointNeighbors.Key;
-			if (InvMasses[Index - Offset] == 0.f)  // TODO: Triangle indices should ideally be starting at 0 to avoid these mix-ups
-			{
-				KinematicIndices.Add(Index);
-			}
-		}
-
-		const TArray<TArray<int32>> IslandElements = FPBDLongRangeConstraints::ComputeIslands(PointToNeighborsMap, KinematicIndices);
-
 		// Draw constraints
 		const FClothConstraints& ClothConstraints = Solver->GetClothConstraints(Offset);
-		
 		const TConstArrayView<FVec3> Positions = Cloth->GetParticlePositions(Solver.Get());
 
 		if (const FPBDLongRangeConstraints* const LongRangeConstraints = ClothConstraints.GetLongRangeConstraints().Get())
 		{
-			for (const FPBDLongRangeConstraints::FTether& Tether : LongRangeConstraints->GetTethers())
+			const TArray<TConstArrayView<FPBDLongRangeConstraints::FTether>>& Tethers = LongRangeConstraints->GetTethers();
+
+			for (int32 BatchIndex = 0; BatchIndex < Tethers.Num(); ++BatchIndex)
 			{
-				const int32 KinematicIndex = Tether.Start;
-				const int32 DynamicIndex = Tether.End;
-				const FReal RefLength = Tether.RefLength;
-				// Find Island
-				int32 ColorIndex = 0;
-				for (int32 IslandIndex = 0; IslandIndex < IslandElements.Num(); ++IslandIndex)
-				{
-					if (IslandElements[IslandIndex].Find(KinematicIndex) != INDEX_NONE)  // TODO: This is O(n^2), it'll be nice to make this faster, even if it is only debug viz. Maybe binary search if the kinematic indices are ordered?
-					{
-						ColorIndex = ColorOffset + IslandIndex;
-						break;
-					}
-				}
-				const FLinearColor Color = PseudoRandomColor(ColorIndex);
-				const FVec3 Pos0 = Positions[KinematicIndex - Offset] + LocalSpaceLocation;
-				const FVec3 Pos1 = Positions[DynamicIndex - Offset] + LocalSpaceLocation;
+				const FLinearColor Color = PseudoRandomColor(ColorOffset + BatchIndex);
 
-				DrawLine(PDI, Pos0, Pos1, Color);
+				const TConstArrayView<FPBDLongRangeConstraints::FTether>& TetherBatch = Tethers[BatchIndex];
 
-				FVec3 Direction = Pos1 - Pos0;
-				const float Length = Direction.SafeNormalize();
-				if (Length > SMALL_NUMBER)
+				// Draw tethers
+				for (const FPBDLongRangeConstraints::FTether& Tether : TetherBatch)
 				{
-					const FVec3 Pos2 = Pos1 + Direction * (Tether.RefLength - Length);
-					DrawLine(PDI, Pos1, Pos2, FLinearColor::Black);  // Color.Desaturate(1.f));
-				}
-			}
-		}
+					const int32 KinematicIndex = LongRangeConstraints->GetStartIndex(Tether);
+					const int32 DynamicIndex = LongRangeConstraints->GetEndIndex(Tether);
+					const FReal TargetLength = LongRangeConstraints->GetTargetLength(Tether);
 
-		// Draw islands
-		const TConstArrayView<TVec3<int32>> Elements = Cloth->GetTriangleMesh(Solver.Get()).GetElements();
+					const FVec3 Pos0 = Positions[KinematicIndex] + LocalSpaceLocation;
+					const FVec3 Pos1 = Positions[DynamicIndex] + LocalSpaceLocation;
 
-		for (int32 ElementIndex = 0; ElementIndex < Elements.Num(); ++ElementIndex)
-		{
-			const auto& Element = Elements[ElementIndex];
+					DrawLine(PDI, Pos0, Pos1, Color);
 
-			const bool bIsKinematic0 = (InvMasses[Element.X - Offset] == 0.f);
-			const bool bIsKinematic1 = (InvMasses[Element.Y - Offset] == 0.f);
-			const bool bIsKinematic2 = (InvMasses[Element.Z - Offset] == 0.f);
-
-			// Lookup for any kinematic point on the triangle element to use for finding the island (it doesn't matter which one, if two kinematic points are on the same triangle they have to be on the same island)
-			const int32 KinematicIndex = bIsKinematic0 ? Element.X : bIsKinematic1 ? Element.Y : bIsKinematic2 ? Element.Z : INDEX_NONE;
-			if (KinematicIndex == INDEX_NONE)
-			{
-				continue;
-			}
-
-			// Find island Color
-			int32 ColorIndex = 0;
-			for (int32 IslandIndex = 0; IslandIndex < IslandElements.Num(); ++IslandIndex)
-			{
-				if (IslandElements[IslandIndex].Find(KinematicIndex) != INDEX_NONE)  // TODO: This is O(n^2), it'll be nice to make this faster, even if it is only debug viz. Maybe binary search if the kinematic indices are ordered?
-				{
-					ColorIndex = ColorOffset + IslandIndex;
-					break;
+					FVec3 Direction = Pos1 - Pos0;
+					const float Length = Direction.SafeNormalize();
+					if (Length > SMALL_NUMBER)
+					{
+						const FVec3 Pos2 = Pos1 + Direction * (TargetLength - Length);
+						DrawLine(PDI, Pos1, Pos2, FLinearColor::Black);  // Color.Desaturate(1.f));
+					}
 				}
 			}
-			const FLinearColor Color = PseudoRandomColor(ColorIndex);
 
-			const FVector Pos0 = LocalSpaceLocation + Positions[Element.X - Offset];
-			const FVector Pos1 = LocalSpaceLocation + Positions[Element.Y - Offset];
-			const FVector Pos2 = LocalSpaceLocation + Positions[Element.Z - Offset];
-
-			if (bIsKinematic0 && bIsKinematic1)
-			{
-				DrawLine(PDI, Pos0, Pos1, Color);
-			}
-			if (bIsKinematic1 && bIsKinematic2)
-			{
-				DrawLine(PDI, Pos1, Pos2, Color);
-			}
-			if (bIsKinematic2 && bIsKinematic0)
-			{
-				DrawLine(PDI, Pos2, Pos0, Color);
-			}
+			// Rotate the colors for each cloth
+			ColorOffset += Tethers.Num();
 		}
-
-		// Rotate the colors for each cloth
-		ColorOffset += IslandElements.Num();
 	}
 }
 
diff --git a/Engine/Plugins/Experimental/ChaosCloth/Source/Chaos/Private/ChaosCloth/ChaosClothingSimulationCloth.cpp b/Engine/Plugins/Experimental/ChaosCloth/Source/Chaos/Private/ChaosCloth/ChaosClothingSimulationCloth.cpp
--- a/Engine/Plugins/Experimental/ChaosCloth/Source/Chaos/Private/ChaosCloth/ChaosClothingSimulationCloth.cpp
+++ b/Engine/Plugins/Experimental/ChaosCloth/Source/Chaos/Private/ChaosCloth/ChaosClothingSimulationCloth.cpp
@@ -16,10 +16,15 @@ namespace ChaosClothingSimulationClothConsoleVariables
 	TAutoConsoleVariable<bool> CVarLegacyDisablesAccurateWind(TEXT("p.ChaosCloth.LegacyDisablesAccurateWind"), true, TEXT("Whether using the Legacy wind model switches off the accurate wind model, or adds up to it"));
 }
 
-FClothingSimulationCloth::FLODData::FLODData(int32 InNumParticles, const TConstArrayView<uint32>& InIndices, const TArray<TConstArrayView<FRealSingle>>& InWeightMaps)
+FClothingSimulationCloth::FLODData::FLODData(
+	int32 InNumParticles,
+	const TConstArrayView<uint32>& InIndices,
+	const TArray<TConstArrayView<FRealSingle>>& InWeightMaps,
+	const TArray<TConstArrayView<TTuple<int32, int32, FRealSingle>>>& InTethers)
 	: NumParticles(InNumParticles)
 	, Indices(InIndices)
 	, WeightMaps(InWeightMaps)
+	, Tethers(InTethers)
 {
 }
 
@@ -119,14 +124,14 @@ void FClothingSimulationCloth::FLODData::Add(FClothingSimulationSolver* Solver,
 
 	// Edge constraints
 	const TConstArrayView<FRealSingle>& EdgeStiffnessMultipliers = WeightMaps[(int32)EChaosWeightMapTarget::EdgeStiffness];
-	if (Cloth->EdgeStiffness[0] > (FReal)0. || (Cloth->EdgeStiffness[1] > (FReal)0. && EdgeStiffnessMultipliers.Num()))
+	if (Cloth->EdgeStiffness[0] > (FRealSingle)0. || (Cloth->EdgeStiffness[1] > (FRealSingle)0. && EdgeStiffnessMultipliers.Num() == NumParticles))
 	{
 		ClothConstraints.SetEdgeConstraints(SurfaceElements, EdgeStiffnessMultipliers, Cloth->bUseXPBDConstraints);
 	}
 
 	// Bending constraints
 	const TConstArrayView<FRealSingle>& BendingStiffnessMultipliers = WeightMaps[(int32)EChaosWeightMapTarget::BendingStiffness];
-	if (Cloth->BendingStiffness[0] > 0.f || (Cloth->BendingStiffness[1] > (FReal)0. && BendingStiffnessMultipliers.Num()))
+	if (Cloth->BendingStiffness[0] > (FRealSingle)0. || (Cloth->BendingStiffness[1] > (FRealSingle)0. && BendingStiffnessMultipliers.Num() == NumParticles))
 	{
 		if (Cloth->bUseBendingElements)
 		{
@@ -142,13 +147,13 @@ void FClothingSimulationCloth::FLODData::Add(FClothingSimulationSolver* Solver,
 
 	// Area constraints
 	const TConstArrayView<FRealSingle>& AreaStiffnessMultipliers = WeightMaps[(int32)EChaosWeightMapTarget::AreaStiffness];
-	if (Cloth->AreaStiffness[0] > 0.f || (Cloth->AreaStiffness[1] > (FReal)0. && AreaStiffnessMultipliers.Num()))
+	if (Cloth->AreaStiffness[0] > (FRealSingle)0. || (Cloth->AreaStiffness[1] > (FRealSingle)0. && AreaStiffnessMultipliers.Num() == NumParticles))
 	{
 		ClothConstraints.SetAreaConstraints(SurfaceElements, AreaStiffnessMultipliers, Cloth->bUseXPBDConstraints);
 	}
 
 	// Volume constraints
-	if (Cloth->VolumeStiffness)
+	if (Cloth->VolumeStiffness > (FRealSingle)0.)
 	{
 		if (Cloth->bUseThinShellVolumeConstraints)
 		{
@@ -192,23 +197,21 @@ void FClothingSimulationCloth::FLODData::Add(FClothingSimulationSolver* Solver,
 	}
 
 	// Long range constraints
-	if (Cloth->TetherStiffness[0] > 0.f || Cloth->TetherStiffness[1] > 0.f)
+	const TConstArrayView<FRealSingle>& TetherStiffnessMultipliers = WeightMaps[(int32)EChaosWeightMapTarget::TetherStiffness];
+	if (Cloth->TetherStiffness[0] > (FRealSingle)0. || (Cloth->TetherStiffness[1] > (FRealSingle)0. && TetherStiffnessMultipliers.Num() == NumParticles))
 	{
-		const TMap<int32, TSet<int32>>& PointToNeighborsMap = TriangleMesh.GetPointToNeighborsMap();
-		const TConstArrayView<FRealSingle>& TetherStiffnessMultipliers = WeightMaps[(int32)EChaosWeightMapTarget::TetherStiffness];
 		const TConstArrayView<FRealSingle>& TetherScaleMultipliers = WeightMaps[(int32)EChaosWeightMapTarget::TetherScale];
 
 		ClothConstraints.SetLongRangeConstraints(
-			PointToNeighborsMap,
+			Tethers,
 			TetherStiffnessMultipliers,
 			TetherScaleMultipliers,
 			FVec2((FReal)Cloth->TetherScale[0], (FReal)Cloth->TetherScale[1]),
-			Cloth->TetherMode,
 			Cloth->bUseXPBDConstraints);
 	}
 
 	// Max distances
-	if (MaxDistances.Num())
+	if (MaxDistances.Num() == NumParticles)
 	{
 		ClothConstraints.SetMaximumDistanceConstraints(MaxDistances);
 	}
@@ -216,21 +219,21 @@ void FClothingSimulationCloth::FLODData::Add(FClothingSimulationSolver* Solver,
 	// Backstop Constraints
 	const TConstArrayView<FRealSingle>& BackstopDistances = WeightMaps[(int32)EChaosWeightMapTarget::BackstopDistance];
 	const TConstArrayView<FRealSingle>& BackstopRadiuses = WeightMaps[(int32)EChaosWeightMapTarget::BackstopRadius];
-	if (BackstopRadiuses.Num() && BackstopDistances.Num())
+	if (BackstopRadiuses.Num() == NumParticles && BackstopDistances.Num() == NumParticles)
 	{
 		ClothConstraints.SetBackstopConstraints(BackstopDistances, BackstopRadiuses, Cloth->bUseLegacyBackstop);
 	}
 
 	// Animation Drive Constraints
 	const TConstArrayView<FRealSingle>& AnimDriveStiffnessMultipliers = WeightMaps[(int32)EChaosWeightMapTarget::AnimDriveStiffness];
-	if (Cloth->AnimDriveStiffness[0] > 0.f || (AnimDriveStiffnessMultipliers.Num() == NumParticles && Cloth->AnimDriveStiffness[1] > 0.f))
+	if (Cloth->AnimDriveStiffness[0] > (FRealSingle)0. || (Cloth->AnimDriveStiffness[1] > (FRealSingle)0. && AnimDriveStiffnessMultipliers.Num() == NumParticles))
 	{
 		const TConstArrayView<FRealSingle>& AnimDriveDampingMultipliers = WeightMaps[(int32)EChaosWeightMapTarget::AnimDriveDamping];
 		ClothConstraints.SetAnimDriveConstraints(AnimDriveStiffnessMultipliers, AnimDriveDampingMultipliers);
 	}
 
 	// Shape target constraint
-	if (Cloth->ShapeTargetStiffness)
+	if (Cloth->ShapeTargetStiffness > (FRealSingle)0.)
 	{
 		ClothConstraints.SetShapeTargetConstraints((FReal)Cloth->ShapeTargetStiffness);
 	}
@@ -274,7 +277,9 @@ void FClothingSimulationCloth::FLODData::Update(FClothingSimulationSolver* Solve
 	ClothConstraints.SetEdgeProperties(FVec2((FReal)Cloth->EdgeStiffness[0], (FReal)Cloth->EdgeStiffness[1]));
 	ClothConstraints.SetBendingProperties(FVec2((FReal)Cloth->BendingStiffness[0], (FReal)Cloth->BendingStiffness[1]));
 	ClothConstraints.SetAreaProperties(FVec2((FReal)Cloth->AreaStiffness[0], (FReal)Cloth->AreaStiffness[1]));
-	ClothConstraints.SetLongRangeAttachmentProperties(FVec2((FReal)Cloth->TetherStiffness[0], (FReal)Cloth->TetherStiffness[1]));
+	ClothConstraints.SetLongRangeAttachmentProperties(
+		FVec2((FReal)Cloth->TetherStiffness[0], (FReal)Cloth->TetherStiffness[1]),
+		FVec2((FReal)Cloth->TetherScale[0], (FReal)Cloth->TetherScale[1]));
 	ClothConstraints.SetSelfCollisionProperties((FReal)Cloth->SelfCollisionThickness);
 	ClothConstraints.SetAnimDriveProperties(
 		FVec2((FReal)Cloth->AnimDriveStiffness[0], (FReal)Cloth->AnimDriveStiffness[1]),
@@ -429,9 +434,14 @@ void FClothingSimulationCloth::SetMesh(FClothingSimulationMesh* InMesh)
 	// Reset LODs
 	const int32 NumLODs = Mesh ? Mesh->GetNumLODs() : 0;
 	LODData.Reset(NumLODs);
+	const bool bUseGeodesicTethers = (TetherMode == ETetherMode::Geodesic);
 	for (int32 Index = 0; Index < NumLODs; ++Index)
 	{
-		LODData.Emplace(Mesh->GetNumPoints(Index), Mesh->GetIndices(Index), Mesh->GetWeightMaps(Index));
+		LODData.Emplace(
+			Mesh->GetNumPoints(Index),
+			Mesh->GetIndices(Index),
+			Mesh->GetWeightMaps(Index),
+			Mesh->GetTethers(Index, bUseGeodesicTethers));
 	}
 
 	// Iterate all known solvers
@@ -627,6 +637,13 @@ const TArray<TConstArrayView<FRealSingle>>& FClothingSimulationCloth::GetWeightM
 	return LODData.IsValidIndex(LODIndex) ? LODData[LODIndex].WeightMaps : EmptyWeightMaps;
 }
 
+const TArray<TConstArrayView<TTuple<int32, int32, float>>>& FClothingSimulationCloth::GetTethers(const FClothingSimulationSolver* Solver) const
+{
+	const int32 LODIndex = LODIndices.FindChecked(Solver);
+	static const TArray<TConstArrayView<TTuple<int32, int32, float>>> EmptyTethers;
+	return LODData.IsValidIndex(LODIndex) ? LODData[LODIndex].Tethers : EmptyTethers;
+}
+
 int32 FClothingSimulationCloth::GetReferenceBoneIndex() const
 {
 	return Mesh ? Mesh->GetReferenceBoneIndex() : INDEX_NONE;
diff --git a/Engine/Plugins/Experimental/ChaosCloth/Source/Chaos/Private/ChaosCloth/ChaosClothingSimulationInteractor.cpp b/Engine/Plugins/Experimental/ChaosCloth/Source/Chaos/Private/ChaosCloth/ChaosClothingSimulationInteractor.cpp
--- a/Engine/Plugins/Experimental/ChaosCloth/Source/Chaos/Private/ChaosCloth/ChaosClothingSimulationInteractor.cpp
+++ b/Engine/Plugins/Experimental/ChaosCloth/Source/Chaos/Private/ChaosCloth/ChaosClothingSimulationInteractor.cpp
@@ -48,21 +48,23 @@ void UChaosClothingInteractor::SetMaterial(FVector2D EdgeStiffness, FVector2D Be
 	}));
 }
 
-void UChaosClothingInteractor::SetLongRangeAttachmentLinear(float TetherStiffnessLinear)
+void UChaosClothingInteractor::SetLongRangeAttachmentLinear(float TetherStiffnessLinear, float TetherScale)
 {
 	// Deprecated
 	const TVec2<FRealSingle> TetherStiffness((FMath::Clamp(FMath::Loge(TetherStiffnessLinear) * ChaosClothingInteractor::InvStiffnessLogBase + 1.f, 0.f, 1.f)), 1.f);
-	Commands.Add(FChaosClothingInteractorCommand::CreateLambda([TetherStiffness](FClothingSimulationCloth* Cloth)
+	Commands.Add(FChaosClothingInteractorCommand::CreateLambda([TetherStiffness, TetherScale](FClothingSimulationCloth* Cloth)
 	{
-		Cloth->SetLongRangeAttachmentProperties(TetherStiffness);
+		Cloth->SetLongRangeAttachmentProperties(TetherStiffness, TVec2<FRealSingle>(TetherScale, TetherScale));
 	}));
 }
 
-void UChaosClothingInteractor::SetLongRangeAttachment(FVector2D TetherStiffness)
+void UChaosClothingInteractor::SetLongRangeAttachment(FVector2D TetherStiffness, FVector2D TetherScale)
 {
-	Commands.Add(FChaosClothingInteractorCommand::CreateLambda([TetherStiffness](FClothingSimulationCloth* Cloth)
+	Commands.Add(FChaosClothingInteractorCommand::CreateLambda([TetherStiffness, TetherScale](FClothingSimulationCloth* Cloth)
 	{
-		Cloth->SetLongRangeAttachmentProperties(TVec2<FRealSingle>(TetherStiffness[0], TetherStiffness[1]));
+		Cloth->SetLongRangeAttachmentProperties(
+			TVec2<FRealSingle>(TetherStiffness[0], TetherStiffness[1]),
+			TVec2<FRealSingle>(TetherScale[0], TetherScale[1]));
 	}));
 }
 
diff --git a/Engine/Plugins/Experimental/ChaosCloth/Source/Chaos/Private/ChaosCloth/ChaosClothingSimulationMesh.cpp b/Engine/Plugins/Experimental/ChaosCloth/Source/Chaos/Private/ChaosCloth/ChaosClothingSimulationMesh.cpp
--- a/Engine/Plugins/Experimental/ChaosCloth/Source/Chaos/Private/ChaosCloth/ChaosClothingSimulationMesh.cpp
+++ b/Engine/Plugins/Experimental/ChaosCloth/Source/Chaos/Private/ChaosCloth/ChaosClothingSimulationMesh.cpp
@@ -97,6 +97,25 @@ TArray<TConstArrayView<FRealSingle>> FClothingSimulationMesh::GetWeightMaps(int3
 	return WeightMaps;
 }
 
+TArray<TConstArrayView<TTuple<int32, int32, float>>> FClothingSimulationMesh::GetTethers(int32 LODIndex, bool bUseGeodesicTethers) const
+{
+	TArray<TConstArrayView<TTuple<int32, int32, float>>> Tethers;
+	if (Asset && Asset->LodData.IsValidIndex(LODIndex))
+	{
+		const FClothLODDataCommon& ClothLODData = Asset->LodData[LODIndex];
+		const FClothPhysicalMeshData& ClothPhysicalMeshData = ClothLODData.PhysicalMeshData;
+		const FClothTetherData& ClothTetherData = bUseGeodesicTethers ? ClothPhysicalMeshData.GeodesicTethers : ClothPhysicalMeshData.EuclideanTethers;
+		
+		const int32 NumTetherBatches = ClothTetherData.Tethers.Num();
+		Tethers.Reserve(NumTetherBatches);
+		for (int32 Index = 0; Index < NumTetherBatches; ++Index)
+		{
+			Tethers.Emplace(TConstArrayView<TTuple<int32, int32, float>>(ClothTetherData.Tethers[Index]));
+		}
+	}
+	return Tethers;
+}
+
 int32 FClothingSimulationMesh::GetReferenceBoneIndex() const
 {
 	return Asset ? Asset->ReferenceBoneIndex : INDEX_NONE;
diff --git a/Engine/Plugins/Experimental/ChaosCloth/Source/Chaos/Public/ChaosCloth/ChaosClothConfig.h b/Engine/Plugins/Experimental/ChaosCloth/Source/Chaos/Public/ChaosCloth/ChaosClothConfig.h
--- a/Engine/Plugins/Experimental/ChaosCloth/Source/Chaos/Public/ChaosCloth/ChaosClothConfig.h
+++ b/Engine/Plugins/Experimental/ChaosCloth/Source/Chaos/Public/ChaosCloth/ChaosClothConfig.h
@@ -70,6 +70,15 @@ public:
 	/** PostLoad override used to deal with updates/changes in properties. */
 	virtual void PostLoad() override;
 
+	/** Return wherether to pre-compute Inverse Masses. */
+	virtual bool NeedsInverseMasses() const override { return false; }  // TODO: Chaos Cloth uses the mass mode enum, and this will require a little refactor to work
+
+	/** Return wherether to pre-compute the Long Range Attachment tethers. */
+	virtual bool NeedsTethers() const override { return TetherStiffness.Low > 0.f || TetherStiffness.High > 0.f; }
+
+	/** Return whether tethers need to be calculated using geodesic distances instead of eclidean. */
+	virtual bool TethersUseGeodesicDistance() const override { return bUseGeodesicDistance; }
+
 	/** Return the mass value, from whichever mass mode (Density, UniformMass, or TotalMass) is selected. */
 	float GetMassValue() const;
 
diff --git a/Engine/Plugins/Experimental/ChaosCloth/Source/Chaos/Public/ChaosCloth/ChaosClothConstraints.h b/Engine/Plugins/Experimental/ChaosCloth/Source/Chaos/Public/ChaosCloth/ChaosClothConstraints.h
--- a/Engine/Plugins/Experimental/ChaosCloth/Source/Chaos/Public/ChaosCloth/ChaosClothConstraints.h
+++ b/Engine/Plugins/Experimental/ChaosCloth/Source/Chaos/Public/ChaosCloth/ChaosClothConstraints.h
@@ -50,9 +50,9 @@ namespace Chaos
 		void SetAreaConstraints(const TArray<TVec3<int32>>& SurfaceElements, const TConstArrayView<FRealSingle>& StiffnessMultipliers, bool bUseXPBDConstraints);
 		void SetVolumeConstraints(const TArray<TVec2<int32>>& DoubleBendingEdges, FReal VolumeStiffness);
 		void SetVolumeConstraints(TArray<TVec3<int32>>&& SurfaceElements, FReal VolumeStiffness);
-		void SetLongRangeConstraints(const TMap<int32, TSet<int32>>& PointToNeighborsMap,
+		void SetLongRangeConstraints(const TArray<TConstArrayView<TTuple<int32, int32, FRealSingle>>>& Tethers,
 			const TConstArrayView<FRealSingle>& TetherStiffnessMultipliers, const TConstArrayView<FRealSingle>& TetherScaleMultipliers,
-			const FVec2& TetherScale, ETetherMode TetherMode, bool bUseXPBDConstraints);
+			const FVec2& TetherScale, bool bUseXPBDConstraints);
 		void SetMaximumDistanceConstraints(const TConstArrayView<FRealSingle>& MaxDistances);
 		void SetBackstopConstraints(const TConstArrayView<FRealSingle>& BackstopDistances, const TConstArrayView<FRealSingle>& BackstopRadiuses, bool bUseLegacyBackstop);
 		void SetAnimDriveConstraints(const TConstArrayView<FRealSingle>& AnimDriveStiffnessMultipliers, const TConstArrayView<FRealSingle>& AnimDriveDampingMultipliers);
@@ -67,7 +67,7 @@ namespace Chaos
 		void SetAreaProperties(const FVec2& AreaStiffness);
 		void SetThinShellVolumeProperties(FReal VolumeStiffness);
 		void SetVolumeProperties(FReal VolumeStiffness);
-		void SetLongRangeAttachmentProperties(const FVec2& TetherStiffness);
+		void SetLongRangeAttachmentProperties(const FVec2& TetherStiffness, const FVec2& TetherScale);
 		void SetMaximumDistanceProperties(FReal MaxDistancesMultiplier);
 		void SetAnimDriveProperties(const FVec2& AnimDriveStiffness, const FVec2& AnimDriveDamping);
 		void SetSelfCollisionProperties(FReal SelfCollisionThickness);
diff --git a/Engine/Plugins/Experimental/ChaosCloth/Source/Chaos/Public/ChaosCloth/ChaosClothingSimulationCloth.h b/Engine/Plugins/Experimental/ChaosCloth/Source/Chaos/Public/ChaosCloth/ChaosClothingSimulationCloth.h
--- a/Engine/Plugins/Experimental/ChaosCloth/Source/Chaos/Public/ChaosCloth/ChaosClothingSimulationCloth.h
+++ b/Engine/Plugins/Experimental/ChaosCloth/Source/Chaos/Public/ChaosCloth/ChaosClothingSimulationCloth.h
@@ -77,7 +77,7 @@ namespace Chaos
 		void SetMaxDistancesMultiplier(FRealSingle InMaxDistancesMultiplier) { MaxDistancesMultiplier = InMaxDistancesMultiplier; }
 
 		void SetMaterialProperties(const TVec2<FRealSingle>& InEdgeStiffness, const TVec2<FRealSingle>& InBendingStiffness, const TVec2<FRealSingle>& InAreaStiffness) { EdgeStiffness = InEdgeStiffness; BendingStiffness = InBendingStiffness; AreaStiffness = InAreaStiffness; }
-		void SetLongRangeAttachmentProperties(const TVec2<FRealSingle>& InTetherStiffness) { TetherStiffness = InTetherStiffness; }
+		void SetLongRangeAttachmentProperties(const TVec2<FRealSingle>& InTetherStiffness, const TVec2<FRealSingle>& InTetherScale) { TetherStiffness = InTetherStiffness; TetherScale = InTetherScale;  }
 		void SetCollisionProperties(FRealSingle InCollisionThickness, FRealSingle InFrictionCoefficient, bool bInUseCCD, FRealSingle InSelfCollisionThickness) { CollisionThickness = InCollisionThickness; FrictionCoefficient = InFrictionCoefficient; bUseCCD = bInUseCCD; SelfCollisionThickness = InSelfCollisionThickness; }
 		void SetDampingProperties(FRealSingle InDampingCoefficient) { DampingCoefficient = InDampingCoefficient; }
 		void SetAerodynamicsProperties(const TVec2<FRealSingle>& InDrag, const TVec2<FRealSingle>& InLift, FRealSingle InAirDensity, const FVec3& InWindVelocity) { Drag = InDrag; Lift = InLift; InAirDensity = AirDensity; WindVelocity = InWindVelocity; }
@@ -118,15 +118,17 @@ namespace Chaos
 		FVec3 GetGravity(const FClothingSimulationSolver* Solver) const;
 		// Return the current bounding box based on a given solver, not thread safe, call must be done right after the solver update.
 		FAABB3 CalculateBoundingBox(const FClothingSimulationSolver* Solver) const;
-		// Return the current LOD Offset in the solver's particle array, or INDEX_NONE if no LOD is currently selected
+		// Return the current LOD offset in the solver's particle array, or INDEX_NONE if no LOD is currently selected.
 		int32 GetOffset(const FClothingSimulationSolver* Solver) const;
-		// Return the current LOD Mesh
+		// Return the current LOD mesh.
 		const FTriangleMesh& GetTriangleMesh(const FClothingSimulationSolver* Solver) const;
-		// Return the current LOD Weightmaps
+		// Return the current LOD weightmaps.
 		const TArray<TConstArrayView<FRealSingle>>& GetWeightMaps(const FClothingSimulationSolver* Solver) const;
-		// Return the reference bone index for this cloth
+		// Return the current LOD tethers.
+		const TArray<TConstArrayView<TTuple<int32, int32, float>>>& GetTethers(const FClothingSimulationSolver* Solver) const;
+		// Return the reference bone index for this cloth.
 		int32 GetReferenceBoneIndex() const;
-		// Return the local reference space transform for this cloth
+		// Return the local reference space transform for this cloth.
 		const FRigidTransform3& GetReferenceSpaceTransform() const { return ReferenceSpaceTransform;  }
 		// ---- End of the debugging/visualization functions
 
@@ -150,6 +152,7 @@ namespace Chaos
 			const int32 NumParticles;
 			const TConstArrayView<uint32> Indices;
 			const TArray<TConstArrayView<FRealSingle>> WeightMaps;
+			const TArray<TConstArrayView<TTuple<int32, int32, FRealSingle>>> Tethers;
 
 			// Per Solver data
 			struct FSolverData
@@ -163,7 +166,11 @@ namespace Chaos
 			int32 NumKinenamicParticles;
 			int32 NumDynammicParticles;
 
-			FLODData(int32 InNumParticles, const TConstArrayView<uint32>& InIndices, const TArray<TConstArrayView<FRealSingle>>& InWeightMaps);
+			FLODData(
+				int32 InNumParticles,
+				const TConstArrayView<uint32>& InIndices,
+				const TArray<TConstArrayView<FRealSingle>>& InWeightMaps,
+				const TArray<TConstArrayView<TTuple<int32, int32, FRealSingle>>>& InTethers);
 
 			void Add(FClothingSimulationSolver* Solver, FClothingSimulationCloth* Cloth, int32 LODIndex);
 			void Remove(FClothingSimulationSolver* Solver);
diff --git a/Engine/Plugins/Experimental/ChaosCloth/Source/Chaos/Public/ChaosCloth/ChaosClothingSimulationInteractor.h b/Engine/Plugins/Experimental/ChaosCloth/Source/Chaos/Public/ChaosCloth/ChaosClothingSimulationInteractor.h
--- a/Engine/Plugins/Experimental/ChaosCloth/Source/Chaos/Public/ChaosCloth/ChaosClothingSimulationInteractor.h
+++ b/Engine/Plugins/Experimental/ChaosCloth/Source/Chaos/Public/ChaosCloth/ChaosClothingSimulationInteractor.h
@@ -30,11 +30,11 @@ public:
 	UFUNCTION(BlueprintCallable, Category = ClothingSimulation, Meta = (Keywords = "Chaos Edge Bending Area Stiffness"))
 	void SetMaterial(FVector2D EdgeStiffness = FVector2D(1.f, 1.f), FVector2D BendingStiffness = FVector2D(1.f, 1.f), FVector2D AreaStiffness = FVector2D(1.f, 1.f));
 
-	UFUNCTION(BlueprintCallable, Category = ClothingSimulation, Meta = (Keywords = "Chaos Tether Stiffness"))
-	void SetLongRangeAttachmentLinear(float TetherStiffness = 1.f);
+	UFUNCTION(BlueprintCallable, Category = ClothingSimulation, Meta = (Keywords = "Chaos Tether Stiffness Scale"))
+	void SetLongRangeAttachmentLinear(float TetherStiffness = 1.f, float TetherScale = 1.f);
 
-	UFUNCTION(BlueprintCallable, Category = ClothingSimulation, Meta = (Keywords = "Chaos Tether Stiffness"))
-	void SetLongRangeAttachment(FVector2D TetherStiffness = FVector2D(0.f, 1.f));
+	UFUNCTION(BlueprintCallable, Category = ClothingSimulation, Meta = (Keywords = "Chaos Tether Stiffness Scale"))
+	void SetLongRangeAttachment(FVector2D TetherStiffness = FVector2D(1.f, 1.f), FVector2D TetherScale = FVector2D(1.f, 1.f));
 
 	UFUNCTION(BlueprintCallable, Category = ClothingSimulation, Meta = (Keywords = "Chaos Self Thickness Friction Coefficient"))
 	void SetCollision(float CollisionThickness = 1.f, float FrictionCoefficient = 0.8f, bool bUseCCD = false, float SelfCollisionThickness = 2.f);
diff --git a/Engine/Plugins/Experimental/ChaosCloth/Source/Chaos/Public/ChaosCloth/ChaosClothingSimulationMesh.h b/Engine/Plugins/Experimental/ChaosCloth/Source/Chaos/Public/ChaosCloth/ChaosClothingSimulationMesh.h
--- a/Engine/Plugins/Experimental/ChaosCloth/Source/Chaos/Public/ChaosCloth/ChaosClothingSimulationMesh.h
+++ b/Engine/Plugins/Experimental/ChaosCloth/Source/Chaos/Public/ChaosCloth/ChaosClothingSimulationMesh.h
@@ -36,8 +36,9 @@ namespace Chaos
 		int32 GetLODIndex() const;
 		int32 GetNumLODs() const;
 		int32 GetNumPoints(int32 LODIndex) const;
-		TConstArrayView<const uint32> GetIndices(int32 LODIndex) const;
-		TArray<TConstArrayView<const FRealSingle>> GetWeightMaps(int32 LODIndex) const;
+		TConstArrayView<uint32> GetIndices(int32 LODIndex) const;
+		TArray<TConstArrayView<FRealSingle>> GetWeightMaps(int32 LODIndex) const;
+		TArray<TConstArrayView<TTuple<int32, int32, float>>> GetTethers(int32 LODIndex, bool bUseGeodesicTethers) const;
 		int32 GetReferenceBoneIndex() const;
 		FRigidTransform3 GetReferenceBoneTransform() const;
 
diff --git a/Engine/Source/Editor/ClothingSystemEditor/Private/ClothingAssetFactory.cpp b/Engine/Source/Editor/ClothingSystemEditor/Private/ClothingAssetFactory.cpp
--- a/Engine/Source/Editor/ClothingSystemEditor/Private/ClothingAssetFactory.cpp
+++ b/Engine/Source/Editor/ClothingSystemEditor/Private/ClothingAssetFactory.cpp
@@ -266,7 +266,7 @@ UClothingAssetBase* UClothingAssetFactory::CreateFromSkeletalMesh(USkeletalMesh*
 		NewAsset->RefreshBoneMapping(TargetMesh);
 
 		// Invalidate cached data as the mesh has changed
-		NewAsset->InvalidateCachedData();
+		NewAsset->InvalidateAllCachedData();
 
 		return NewAsset;
 	}
@@ -290,7 +290,7 @@ UClothingAssetBase* UClothingAssetFactory::CreateFromExistingCloth(USkeletalMesh
 	//Need to empty LODMap to remove previous mappings from cloth LOD to SkelMesh LOD
 	NewAsset->LodMap.Empty();
 	NewAsset->RefreshBoneMapping(TargetMesh);
-	NewAsset->InvalidateCachedData();
+	NewAsset->InvalidateAllCachedData();
 
 	return NewAsset;
 }
@@ -374,7 +374,7 @@ UClothingAssetBase* UClothingAssetFactory::ImportLodToClothing(USkeletalMesh* Ta
 				ConcreteTarget->BuildLodTransitionData();
 
 				// Invalidate cached data as the mesh has changed
-				ConcreteTarget->InvalidateCachedData();
+				ConcreteTarget->InvalidateAllCachedData();
 
 				return ConcreteTarget;
 			}
@@ -463,10 +463,9 @@ UClothingAssetBase* UClothingAssetFactory::CreateFromApexAsset(nvidia::apex::Clo
 	FName ConvertedBoneName(*FString(InApexAsset->getBoneName(AssetInternalRootBoneIndex)).Replace(TEXT(" "), TEXT("-")));
 
 	NewClothingAsset->AssetGuid = FGuid::NewGuid();
-	NewClothingAsset->InvalidateCachedData();
+	NewClothingAsset->InvalidateCachedData(ClothingCachedData_All);
 
 	NewClothingAsset->BuildLodTransitionData();
-	NewClothingAsset->BuildSelfCollisionData();
 	NewClothingAsset->CalculateReferenceBoneIndex();
 
 	// Add masks for parameters
diff --git a/Engine/Source/Runtime/ClothingSystemRuntimeCommon/Private/ClothPhysicalMeshData.cpp b/Engine/Source/Runtime/ClothingSystemRuntimeCommon/Private/ClothPhysicalMeshData.cpp
--- a/Engine/Source/Runtime/ClothingSystemRuntimeCommon/Private/ClothPhysicalMeshData.cpp
+++ b/Engine/Source/Runtime/ClothingSystemRuntimeCommon/Private/ClothPhysicalMeshData.cpp
@@ -2,6 +2,8 @@
 #include "ClothPhysicalMeshData.h"
 #include "ClothConfigBase.h"
 #include "ClothPhysicalMeshDataBase_Legacy.h"
+#include "GPUSkinPublicDefs.h"  // For MAX_TOTAL_INFLUENCES
+#include "ClothTetherData.h"
 
 FClothPhysicalMeshData::FClothPhysicalMeshData()
 	: MaxBoneWeights(0)
@@ -85,20 +87,21 @@ void FClothPhysicalMeshData::ClearWeightMaps()
 
 void FClothPhysicalMeshData::BuildSelfCollisionData(const TMap<FName, TObjectPtr<UClothConfigBase>>& ClothConfigs)
 {
-#if WITH_APEX_CLOTHING  // Only apex clothing needs to build the SelfCollisionIndices
-	float SCRadius = 0.f;
+	// Deprecated from 5.0 onwards
+	float SelfCollisionRadius = 0.f;
 	for (const TPair<FName, TObjectPtr<UClothConfigBase>>& ClothConfig : ClothConfigs)
 	{
-		SCRadius = ClothConfig.Value->NeedsSelfCollisionIndices();
-		if (SCRadius > 0.f)
-		{
-			break; 
-		}
+		SelfCollisionRadius = FMath::Max(SelfCollisionRadius, ClothConfig.Value->NeedsSelfCollisionData());
 	}
+	if (SelfCollisionRadius)
+	{
+		BuildSelfCollisionData(SelfCollisionRadius);
+	}
+}
 
-	if (SCRadius == 0.f) { return; }
-
-	const float SCRadiusSq = SCRadius * SCRadius;
+void FClothPhysicalMeshData::BuildSelfCollisionData(float SelfCollisionRadius)
+{
+	const float SelfCollisionRadiusSq = SelfCollisionRadius * SelfCollisionRadius;
 
 	// Start with the full set
 	const int32 NumVerts = Vertices.Num();
@@ -136,7 +139,7 @@ void FClothPhysicalMeshData::BuildSelfCollisionData(const TMap<FName, TObjectPtr
 
 			const FVector& V1Pos = Vertices[V1Index];
 			const float V0ToV1DistSq = (V1Pos - V0Pos).SizeSquared();
-			if (V0ToV1DistSq < SCRadiusSq)
+			if (V0ToV1DistSq < SelfCollisionRadiusSq)
 			{
 				// Points are in contact in the rest state.  Remove it.
 				//
@@ -157,5 +160,111 @@ void FClothPhysicalMeshData::BuildSelfCollisionData(const TMap<FName, TObjectPtr
 			SelfCollisionIndices.RemoveAt(It);
 		}
 	}
-#endif  // WITH_APEX_CLOTHING
+}
+
+void FClothPhysicalMeshData::CalculateInverseMasses()
+{
+	// Recalculate inverse masses for the physical mesh particles
+	check(Indices.Num() % 3 == 0);
+
+	const int32 NumVerts = Vertices.Num();
+	InverseMasses.Empty(NumVerts);
+	InverseMasses.AddZeroed(NumVerts);
+
+	for (int32 TriBaseIndex = 0; TriBaseIndex < Indices.Num(); TriBaseIndex += 3)
+	{
+		const int32 Index0 = Indices[TriBaseIndex];
+		const int32 Index1 = Indices[TriBaseIndex + 1];
+		const int32 Index2 = Indices[TriBaseIndex + 2];
+
+		const FVector AB = Vertices[Index1] - Vertices[Index0];
+		const FVector AC = Vertices[Index2] - Vertices[Index0];
+		const float TriArea = FVector::CrossProduct(AB, AC).Size();
+
+		InverseMasses[Index0] += TriArea;
+		InverseMasses[Index1] += TriArea;
+		InverseMasses[Index2] += TriArea;
+	}
+
+	NumFixedVerts = 0;
+
+	const FPointWeightMap* const MaxDistances = FindWeightMap(EWeightMapTargetCommon::MaxDistance);
+	const TFunction<bool(int32)> IsKinematic = (!MaxDistances || !MaxDistances->Num()) ?
+		TFunction<bool(int32)>([](int32)->bool { return false; }) :
+		TFunction<bool(int32)>([&MaxDistances](int32 Index)->bool { return (*MaxDistances)[Index] < SMALL_NUMBER; });  // For consistency, the default Threshold should be 0.1, not SMALL_NUMBER. But for backward compatibility it needs to be SMALL_NUMBER for now.
+
+	float MassSum = 0.0f;
+	for (int32 CurrVertIndex = 0; CurrVertIndex < NumVerts; ++CurrVertIndex)
+	{
+		float& InverseMass = InverseMasses[CurrVertIndex];
+
+		if (IsKinematic(CurrVertIndex))
+		{
+			InverseMass = 0.0f;
+			++NumFixedVerts;
+		}
+		else
+		{
+			MassSum += InverseMass;
+		}
+	}
+
+	if (MassSum > 0.0f)
+	{
+		const float MassScale = (float)(NumVerts - NumFixedVerts) / MassSum;
+		for (float& InverseMass : InverseMasses)
+		{
+			if (InverseMass != 0.0f)
+			{
+				InverseMass *= MassScale;
+				InverseMass = 1.0f / InverseMass;
+			}
+		}
+	}
+	// TODO: Cache config base (Chaos) mass data
+	//       Note that multiple configs with different mass modes will require a different
+	//       structure than this one to store the mass.
+}
+
+void FClothPhysicalMeshData::CalculateNumInfluences()
+{
+	// Needed by all cloth implementations got skinning
+	for (int32 VertIndex = 0; VertIndex < Vertices.Num(); ++VertIndex)
+	{
+		FClothVertBoneData& BoneDatum = BoneData[VertIndex];
+		const uint16* BoneIndices = BoneDatum.BoneIndices;
+		const float* BoneWeights = BoneDatum.BoneWeights;
+
+		BoneDatum.NumInfluences = MAX_TOTAL_INFLUENCES;
+
+		int32 NumInfluences = 0;
+		for (int32 InfluenceIndex = 0; InfluenceIndex < MAX_TOTAL_INFLUENCES; ++InfluenceIndex)
+		{
+			if (BoneWeights[InfluenceIndex] == 0.0f || BoneIndices[InfluenceIndex] == INDEX_NONE)
+			{
+				BoneDatum.NumInfluences = NumInfluences;
+				break;
+			}
+			++NumInfluences;
+		}
+	}
+}
+
+void FClothPhysicalMeshData::CalculateTethers(bool bUseEuclideanDistance, bool bUseGeodesicDistance)
+{
+	if (const FPointWeightMap* const MaxDistances = FindWeightMap(EWeightMapTargetCommon::MaxDistance))
+	{
+		// Note: Technically there could be two different configs requiring different flavour of tethers,
+		//       in reality at the time of writing there is only one type of tether involved at any one type.
+		if (bUseEuclideanDistance)
+		{
+			constexpr bool bGenerateGeodesic = false;
+			EuclideanTethers.GenerateTethers(Vertices, Indices, MaxDistances->Values, bGenerateGeodesic);
+		}
+		if (bUseGeodesicDistance)
+		{
+			constexpr bool bGenegateGeodesic = true;
+			GeodesicTethers.GenerateTethers(Vertices, Indices, MaxDistances->Values, bGenegateGeodesic);
+		}
+	}
 }
diff --git a/Engine/Source/Runtime/ClothingSystemRuntimeCommon/Private/ClothTetherData.cpp b/Engine/Source/Runtime/ClothingSystemRuntimeCommon/Private/ClothTetherData.cpp
new file mode 100644
--- /dev/null
+++ b/Engine/Source/Runtime/ClothingSystemRuntimeCommon/Private/ClothTetherData.cpp
@@ -0,0 +1,503 @@
+// Copyright Epic Games, Inc. All Rights Reserved.
+#include "ClothTetherData.h"
+#include "Async/ParallelFor.h"
+
+class FClothTetherDataPrivate
+{
+public:
+	static constexpr int32 MaxNumAttachments = 4;  // Max recommended number of tethers per point, could eventually been specified by config
+
+	typedef TTuple<int32, int32, float> FTether;
+
+	FClothTetherDataPrivate(
+		const TConstArrayView<FVector3f>& Points,
+		const TConstArrayView<uint32>& Indices,
+		const TConstArrayView<float>& MaxDistances,
+		bool bUseGeodesicDistance);
+
+	const TArray<TArray<FTether>>& GetTetherSlots() const { return TetherSlots; }
+
+	const TArray<int32>& GetTetherNums() const { return TetherNums; }
+
+private:
+
+	// Generate an array of unconnected islands from a selection of points.
+	void ComputeKinematicNodeIslands();
+
+	// Generate a map of tethers by following the triangle mesh network
+	// The choice of the tether is determined by the shortest euclidean (beeline) distance.
+	void GenerateEuclideanTethers(const TConstArrayView<FVector3f>& Points);
+
+	// Generate a map of tethers by following the triangle mesh network
+	// The choice of the tether is determined by the shortest geodesic (curvature) distance.
+	void GenerateGeodesicTethers(const TConstArrayView<FVector3f>& Points, const TConstArrayView<float>& MaxDistances);
+
+private:
+	TMap<int32, TSet<int32>> NodeToNeighbors;
+	TArray<int32> DynamicNodes;
+	TArray<int32> KinematicNodes;
+	TArray<TArray<int32>> KinematicNodeIslands;
+	TArray<TArray<FTether>> TetherSlots;  // Each dynamic node has between 0 and MaxNumAttachments slots
+	TArray<int32> TetherNums;
+};
+
+void FClothTetherData::GenerateTethers(
+	const TConstArrayView<FVector3f>& Points,
+	const TConstArrayView<uint32>& Indices,
+	const TConstArrayView<float>& MaxDistances,
+	bool bUseGeodesicDistance)
+{
+	// Early exit if there is no MaxDistances mask
+	if (Points.Num() != MaxDistances.Num())
+	{
+		Tethers.Empty();
+		return;
+	}
+
+	// Calculate the tethers
+	const FClothTetherDataPrivate ClothTetherData(Points, Indices, MaxDistances, bUseGeodesicDistance);
+
+	const TArray<TArray<FClothTetherDataPrivate::FTether>>& TetherSlots = ClothTetherData.GetTetherSlots();
+	const TArray<int32> TetherNums = ClothTetherData.GetTetherNums();
+
+	// Reorganize the multiple tethers per node array into a single sequencial batches of tethers more suitable for parallel processing
+	const int32 NumDynamicNodes = TetherSlots.Num();
+	const int32 MaxNumUsedSlots = TetherNums.Num();
+	Tethers.Reset(MaxNumUsedSlots);
+	Tethers.SetNum(MaxNumUsedSlots);
+
+	for (int32 Slot = 0; Slot < MaxNumUsedSlots; ++Slot)
+	{
+		Tethers[Slot].Reserve(TetherNums[Slot]);
+
+		for (int32 Index = 0; Index < NumDynamicNodes; ++Index)
+		{
+			if (TetherSlots[Index].IsValidIndex(Slot))
+			{
+				Tethers[Slot].Emplace(TetherSlots[Index][Slot]);
+			}
+		}
+	}
+}
+
+bool FClothTetherData::Serialize(FArchive& Ar)
+{
+	// Serialize normal tagged property data
+	if (Ar.IsLoading() || Ar.IsSaving())
+	{
+		UScriptStruct* const Struct = FClothTetherData::StaticStruct();
+		Struct->SerializeTaggedProperties(Ar, (uint8*)this, Struct, nullptr);
+	}
+
+	// Serialize the tethers
+	Ar << Tethers;
+
+	return true;
+}
+
+FClothTetherDataPrivate::FClothTetherDataPrivate(
+	const TConstArrayView<FVector3f>& Points,
+	const TConstArrayView<uint32>& Indices,
+	const TConstArrayView<float>& MaxDistances,
+	bool bUseGeodesicDistance)
+{
+	// Calculate the points' neighbors map
+	NodeToNeighbors.Reserve(Points.Num());
+
+	for (int32 i = 0; i < Indices.Num(); i += 3)
+	{
+		const int32 Index0 = Indices[i];
+		const int32 Index1 = Indices[i + 1];
+		const int32 Index2 = Indices[i + 2];
+
+		NodeToNeighbors.FindOrAdd(Index0).Append({ Index1, Index2 });
+		NodeToNeighbors.FindOrAdd(Index1).Append({ Index0, Index2 });
+		NodeToNeighbors.FindOrAdd(Index2).Append({ Index0, Index1 });
+	}
+
+	// Fill up the list of all used indices
+	TArray<int32> Nodes;
+	NodeToNeighbors.GenerateKeyArray(Nodes);
+
+	// Find all kinematic points to use as anchor points for the tethers and seed the path finding
+	DynamicNodes.Reserve(2 * (Nodes.Num() / 3));  // Start at 66% of the number of nodes to minimize this array's reallocations
+	KinematicNodes.Reserve(Nodes.Num() / 3);  // Start at 33% of the number of nodes to minimize this array's reallocations
+	for (const int32 Node : Nodes)
+	{
+		if (MaxDistances[Node])
+		{
+			DynamicNodes.Add(Node);
+		}
+		else
+		{
+			KinematicNodes.Add(Node);
+		}
+	}
+
+	// Compute the islands of kinematic particles
+	ComputeKinematicNodeIslands();
+
+	// Allocate the tether batches, each node has a maximum of MaxNumAttachments slots
+	TetherSlots.SetNum(DynamicNodes.Num());
+
+	for (TArray<FTether>& TetherSlot : TetherSlots)
+	{
+		TetherSlot.Reserve(MaxNumAttachments);
+	}
+
+	// Find the tethers
+	if (!bUseGeodesicDistance)
+	{
+		 GenerateEuclideanTethers(Points);
+	}
+	else
+	{
+		 GenerateGeodesicTethers(Points, MaxDistances);
+	}
+
+	// Update counts
+	TetherNums.Reserve(MaxNumAttachments);
+	for (const TArray<FTether>& TetherSlot : TetherSlots)
+	{
+		// Resize the count array whenever needed
+		const int32 NumUsedSlots = TetherSlot.Num();
+		TetherNums.SetNum(FMath::Max(TetherNums.Num(), NumUsedSlots));
+
+		// Increment each used slot count
+		for (int32 UsedSlot = 0; UsedSlot < NumUsedSlots; ++UsedSlot)
+		{
+			++TetherNums[UsedSlot];
+		}
+	}
+}
+
+void FClothTetherDataPrivate::ComputeKinematicNodeIslands()
+{
+	KinematicNodeIslands.Reset();
+
+	int32 NextIslandIndex = 0;
+	TArray<int32> FreeIslandIndices;
+
+	TMap<int32, int32> NodeToIslands;
+	NodeToIslands.Reserve(KinematicNodes.Num());
+
+	for (const int32 KinematicNode : KinematicNodes)
+	{
+		// Assign this point an island, possibly unionizing existing islands
+		int32 IslandIndex = TNumericLimits<int32>::Max();
+
+		const TSet<int32>& NeighborNodes = NodeToNeighbors[KinematicNode];
+		for (const int32 NeighborNode : NeighborNodes)
+		{
+			if (const int32* const NeighborIslandIndexPtr = NodeToIslands.Find(NeighborNode))
+			{
+				const int32 NeighborIslandIndex = *NeighborIslandIndexPtr;
+
+				if (IslandIndex == TNumericLimits<int32>::Max())
+				{
+					// No assigned island yet, join with the neighbor's island
+					IslandIndex = NeighborIslandIndex;
+				}
+				else if (NeighborIslandIndex != IslandIndex)
+				{
+					// This point is connected to multiple islands, so union them
+					TArray<int32>& NeighborIsland = KinematicNodeIslands[NeighborIslandIndex];
+					for (const int32 NeighborIslandNode : NeighborIsland)
+					{
+						check(NodeToIslands[NeighborIslandNode] == NeighborIslandIndex);
+						NodeToIslands[NeighborIslandNode] = IslandIndex;
+					}
+					KinematicNodeIslands[IslandIndex].Append(NeighborIsland); // Union the two islands
+					NeighborIsland.Reset();  // Don't deallocate, to allow for reusing the free island
+					FreeIslandIndices.AddUnique(NeighborIslandIndex);
+				}
+				// Else this neighbor is already in the same island as the previous neighbors
+			}
+		}
+
+		// If no connected IslandIndex was found, create a new one (or reuse an old vacated one)
+		if (IslandIndex == TNumericLimits<int32>::Max())
+		{
+			if (!FreeIslandIndices.Num())
+			{
+				IslandIndex = NextIslandIndex++;
+				KinematicNodeIslands.SetNum(NextIslandIndex);
+			}
+			else
+			{
+				// Reuse a previously allocated, but currently unused, island
+				IslandIndex = FreeIslandIndices.Pop();
+			}
+		}
+
+		NodeToIslands.FindOrAdd(KinematicNode) = IslandIndex;
+		check(KinematicNodeIslands.IsValidIndex(IslandIndex));
+		KinematicNodeIslands[IslandIndex].Add(KinematicNode);
+	}
+
+	// Remove empty KinematicNodeIslands
+	int32 TotalNumIslandKinematicNodes = 0;
+	for (int32 IslandIndex = 0; IslandIndex < KinematicNodeIslands.Num(); )
+	{
+		const int32 NumIslandKinematicNodes = KinematicNodeIslands[IslandIndex].Num();
+		if (!NumIslandKinematicNodes)
+		{
+			constexpr bool bAllowShrinking = false;
+			KinematicNodeIslands.RemoveAtSwap(IslandIndex, 1, bAllowShrinking);
+			// RemoveAtSwap takes the last elements to replace the current one, do not increment the index in this case
+		}
+		else
+		{
+			TotalNumIslandKinematicNodes += NumIslandKinematicNodes;
+			++IslandIndex;
+		}
+	}
+
+	// Final sanity check
+	check(TotalNumIslandKinematicNodes == KinematicNodes.Num());
+}
+
+void FClothTetherDataPrivate::GenerateEuclideanTethers(const TConstArrayView<FVector3f>& Points)
+{
+	check(KinematicNodeIslands.Num());
+
+	// For each dynamic node
+	ParallelFor(DynamicNodes.Num(), [this, &Points](int32 Index)
+	{
+		const int32 DynamicNode = DynamicNodes[Index];
+
+		// Measure the distance to all kinematic nodes, and keep the closest from each island
+		TArray<TPair<float, int32>> ClosestKinematicNodes;
+		ClosestKinematicNodes.Reserve(KinematicNodeIslands.Num());
+
+		for (const TArray<int32>& KinematicNodeIsland : KinematicNodeIslands)
+		{
+			check(KinematicNodeIsland.Num());  // Island must not be an empty array
+
+			int32 ClosestKinematicNode = TNumericLimits<int32>::Max();
+			float ClosestSquareDistance = TNumericLimits<float>::Max();
+			for (const int32 KinematicNode : KinematicNodeIsland)
+			{
+				const float SquareDistance = (Points[KinematicNode] - Points[DynamicNode]).SizeSquared();
+				if (SquareDistance < ClosestSquareDistance)
+				{
+					ClosestKinematicNode = KinematicNode;
+					ClosestSquareDistance = SquareDistance;
+				}
+			}
+			check(ClosestKinematicNode != TNumericLimits<int32>::Max());
+
+			ClosestKinematicNodes.Emplace(FMath::Sqrt(ClosestSquareDistance), ClosestKinematicNode);
+		}
+
+		// Only keep the first MaxNumAttachments closest kinematic nodes
+		if (ClosestKinematicNodes.Num() > MaxNumAttachments)
+		{
+			// Order all by distance, smallest first
+			ClosestKinematicNodes.Sort();
+
+			// Shrink the list...
+			constexpr bool bAllowShrinking = false;  // ... but not the array
+			ClosestKinematicNodes.SetNum(MaxNumAttachments, bAllowShrinking);
+		}
+
+		// Finally create the tethers between this dynamic node and the N closests kinematic ones
+		TetherSlots[Index].Reserve(ClosestKinematicNodes.Num());
+
+		for (const TPair<float, int32> ClosestKinematicNode : ClosestKinematicNodes)
+		{
+			const float Length = ClosestKinematicNode.Key;
+			const int32 KinematicNode = ClosestKinematicNode.Value;
+
+			TetherSlots[Index].Emplace(KinematicNode, DynamicNode, Length);
+		}
+	});
+}
+
+void FClothTetherDataPrivate::GenerateGeodesicTethers(const TConstArrayView<FVector3f>& Points, const TConstArrayView<float>& MaxDistances)
+{
+	check(KinematicNodeIslands.Num());
+
+	// Find all seeds in each island, kinematic nodes connected to at least one dynamic node
+	int32 NumSeeds = 0;
+	TArray<TArray<int32>> SeedIslands;
+	SeedIslands.SetNum(KinematicNodeIslands.Num());
+
+	for (int32 IslandIndex = 0; IslandIndex < KinematicNodeIslands.Num(); ++IslandIndex)
+	{
+		const TArray<int32>& KinematicNodeIsland = KinematicNodeIslands[IslandIndex];
+		SeedIslands[IslandIndex].Reserve(KinematicNodeIsland.Num());
+
+		for (const int32 KinematicNode : KinematicNodeIsland)
+		{
+			const TSet<int32>& Neighbors = NodeToNeighbors[KinematicNode];
+			for (const int32 Neighbor : Neighbors)
+			{
+				if (MaxDistances[Neighbor])
+				{
+					SeedIslands[IslandIndex].Add(KinematicNode);
+					++NumSeeds;
+					break;
+				}
+			}
+		}
+	}
+
+	// Consolidate into a single array for ease of iteration
+	TArray<int32> Seeds;
+	Seeds.Reserve(NumSeeds);
+
+	for (const TArray<int32>& SeedIsland : SeedIslands)
+	{
+		Seeds.Append(SeedIsland);
+	}
+	check(Seeds.Num() == NumSeeds);
+
+	// Dijkstra for each Kinematic Particle (assume a small number of kinematic points) - note this is N^2 log N with N kinematic points
+	// Find shortest paths from each seed node to all dynamic nodes
+	TMap<int32, TMap<int32, float>> GeodesicDistances;  // Uses two maps instead of one map addressed by Seed and Node to avoid data races
+	GeodesicDistances.Reserve(Seeds.Num());
+
+	for (const int32 Seed : Seeds)
+	{
+		TMap<int32, float>& SeedGeodesicDistances = GeodesicDistances.Emplace(Seed);
+		SeedGeodesicDistances.Reserve(DynamicNodes.Num());
+
+		for (const int32 DynamicNode : DynamicNodes)
+		{
+			SeedGeodesicDistances.Emplace(DynamicNode, TNumericLimits<float>::Max());
+		}
+	}
+
+	ParallelFor(Seeds.Num(), [this, &Seeds, &GeodesicDistances, &Points, &MaxDistances](int32 Index)
+	{
+		const int32 Seed = Seeds[Index];
+		TMap<int32, float>& SeedGeodesicDistances = GeodesicDistances[Seed];
+
+		// Keep track of all visited nodes in a bit array
+		TBitArray<> VisitedNodes(false, Points.Num());
+
+		// Priority queue based implementation of the Dijkstra algorithm
+#define CLOTHTETHERDATA_PROFILE_HEAPSIZE 0
+#if CLOTHTETHERDATA_PROFILE_HEAPSIZE
+		static int32 MaxHeapSize = 1;  // Record the max heap size
+#else
+		static const int32 MaxHeapSize = 512;  // Set the queue size to something large enough to avoid reallocations in most cases
+#endif
+		auto LessPredicate = [&SeedGeodesicDistances](int32 Node1, int32 Node2) -> bool
+			{
+				return SeedGeodesicDistances[Node1] < SeedGeodesicDistances[Node2];  // Less for node priority
+			};
+
+		TArray<int32> Queue;
+		Queue.Reserve(MaxHeapSize);
+		Queue.Heapify(LessPredicate);  // Turn the array into a priority queue
+
+		// Initiate the graph progression
+		VisitedNodes[Seed] = true;
+		Queue.HeapPush(Seed, LessPredicate);
+
+		do
+		{
+			int32 ParentNode;
+			Queue.HeapPop(ParentNode, LessPredicate, false);
+
+			check(VisitedNodes[ParentNode]);
+
+			const float ParentDistance = (ParentNode != Seed) ? SeedGeodesicDistances[ParentNode] : 0.f;
+
+			const TSet<int32>& NeighborNodes = NodeToNeighbors[ParentNode];
+			for (const int32 NeighborNode : NeighborNodes)
+			{
+				check(NeighborNode != ParentNode);
+
+				// Do not progress onto kinematic nodes
+				if (!MaxDistances[NeighborNode])
+				{
+					continue;
+				}
+
+				// Update the geodesic distance if this path is a shorter one
+				const float NewDistance = ParentDistance + (Points[NeighborNode] - Points[ParentNode]).Size();
+
+				float& GeodesicDistance = SeedGeodesicDistances[NeighborNode];
+
+				if (NewDistance < GeodesicDistance)
+				{
+					// Update this path distance
+					GeodesicDistance = NewDistance;
+
+					// Progress to this node position if it hasn't yet been visited
+					if (!VisitedNodes[NeighborNode])
+					{
+						VisitedNodes[NeighborNode] = true;
+
+						Queue.HeapPush(NeighborNode, LessPredicate);
+
+#if CLOTHTETHERDATA_PROFILE_HEAPSIZE
+						MaxHeapSize = FMath::Max(Queue.Num(), MaxHeapSize);
+#endif
+					}
+				}
+			}
+		} while (Queue.Num());
+	});
+
+	// Initialize the tether slot arrays
+	TetherSlots.SetNum(DynamicNodes.Num());
+
+	for (TArray<FTether>& NewTethers : TetherSlots)
+	{
+		NewTethers.Reserve(MaxNumAttachments);
+	}
+
+	// Find the tether constraints starting from each dynamic node
+	ParallelFor(DynamicNodes.Num(), [this, &SeedIslands, &GeodesicDistances](int32 Index)
+	{
+		const int32 DynamicNode = DynamicNodes[Index];
+
+		// Find the closest seeds in each island
+		TArray<int32, TInlineAllocator<32>> ClosestSeeds;
+		ClosestSeeds.Reserve(KinematicNodeIslands.Num());
+
+		for (const TArray<int32>& SeedIsland : SeedIslands)
+		{
+			int32 ClosestSeed = INDEX_NONE;
+			float ClosestDistance = TNumericLimits<float>::Max();
+
+			for (const int32 Seed : SeedIsland)
+			{
+				const float Distance = GeodesicDistances[Seed][DynamicNode];
+				if (Distance < ClosestDistance)
+				{
+					ClosestDistance = Distance;
+					ClosestSeed = Seed;
+				}
+			}
+			if (ClosestSeed != INDEX_NONE)
+			{
+				ClosestSeeds.Add(ClosestSeed);
+			}
+		}
+
+		// Sort all the tethers for this node based on smallest distance
+		ClosestSeeds.Sort([&GeodesicDistances, DynamicNode](int32 Seed1, int32 Seed2)
+			{
+				return GeodesicDistances[Seed1][DynamicNode] < GeodesicDistances[Seed2][DynamicNode];
+			});
+
+		// Keep only N closest tethers
+		if (ClosestSeeds.Num() > MaxNumAttachments)
+		{
+			ClosestSeeds.SetNum(MaxNumAttachments);
+		}
+
+		// Add these tethers to the N (or less) available slots
+		for (const int32 Seed : ClosestSeeds)
+		{
+			const float RefLength = GeodesicDistances[Seed][DynamicNode];
+			TetherSlots[Index].Emplace(Seed, DynamicNode, RefLength);
+		}
+		check(TetherSlots[Index].Num() <= MaxNumAttachments);
+	});
+}
diff --git a/Engine/Source/Runtime/ClothingSystemRuntimeCommon/Private/ClothingAsset.cpp b/Engine/Source/Runtime/ClothingSystemRuntimeCommon/Private/ClothingAsset.cpp
--- a/Engine/Source/Runtime/ClothingSystemRuntimeCommon/Private/ClothingAsset.cpp
+++ b/Engine/Source/Runtime/ClothingSystemRuntimeCommon/Private/ClothingAsset.cpp
@@ -25,6 +25,7 @@
 #include "ComponentReregisterContext.h"
 #include "UObject/UObjectIterator.h"
 #include "UObject/PhysicsObjectVersion.h"
+#include "UObject/UE5PrivateFrostyStreamObjectVersion.h"
 
 #include "GPUSkinPublicDefs.h"
 #include "GPUSkinVertexFactory.h"
@@ -182,7 +183,7 @@ UClothingAssetCommon::UClothingAssetCommon(const FObjectInitializer& ObjectIniti
 	, ChaosClothSimConfig_DEPRECATED(nullptr)
 #endif
 	, ReferenceBoneIndex(0)
-	, CustomData(nullptr)
+	, CustomData(nullptr)  // Deprecated
 {
 }
 
@@ -282,7 +283,6 @@ bool UClothingAssetCommon::BindToSkeletalMesh(
 		return false;
 	}
 
-	BuildSelfCollisionData();
 	CalculateReferenceBoneIndex();
 
 	// Grab the clothing and skel lod data
@@ -419,10 +419,12 @@ bool UClothingAssetCommon::BindToSkeletalMesh(
 		InSkelMesh->GetRefSkeleton().EnsureParentsExistAndSort(SkelLod.ActiveBoneIndices);
 	}
 
-	if(CustomData)
+PRAGMA_DISABLE_DEPRECATION_WARNINGS
+	if(CustomData)  // Deprecated from 5.0 onward
 	{
 		CustomData->BindToSkeletalMesh(InSkelMesh, InMeshLodIndex, InSectionIndex, InAssetLodIndex);
 	}
+PRAGMA_ENABLE_DEPRECATION_WARNINGS
 
 	// Make sure the LOD map is always big enough for the asset to use.
 	// This shouldn't grow to an unwieldy size but maybe consider compacting later.
@@ -549,7 +551,8 @@ void UClothingAssetCommon::ApplyParameterMasks(bool bUpdateFixedVertData)
 	{
 		Lod.PushWeightsToMesh();
 	}
-	InvalidateCachedData();
+	// Invalidate all cached data that depend on masks
+	InvalidateFlaggedCachedData(EClothingCachedDataFlagsCommon::InverseMasses | EClothingCachedDataFlagsCommon::Tethers);
 
 	// Recompute weights if needed
 	USkeletalMesh* const SkeletalMesh = Cast<USkeletalMesh>(GetOuter());
@@ -769,17 +772,6 @@ int32 UClothingAssetCommon::GetNumLods() const
 	return LodData.Num();
 }
 
-void UClothingAssetCommon::BuildSelfCollisionData()
-{
-	if (ClothConfigs.Num())
-	{
-		for(FClothLODDataCommon& Lod : LodData)
-		{
-			Lod.PhysicalMeshData.BuildSelfCollisionData(ClothConfigs);
-		}
-	}
-}
-
 void UClothingAssetCommon::PostLoad()
 {
 	Super::PostLoad();
@@ -799,12 +791,10 @@ void UClothingAssetCommon::PostLoad()
 		}
 	}
 	ClothLodData_DEPRECATED.Empty();
-#endif
 
 	const int32 AnimPhysCustomVersion = GetLinkerCustomVersion(FAnimPhysObjectVersion::GUID);
 	if (AnimPhysCustomVersion < FAnimPhysObjectVersion::AddedClothingMaskWorkflow)
 	{
-#if WITH_EDITORONLY_DATA
 		// Convert current parameters to masks
 		for (FClothLODDataCommon& Lod : LodData)
 		{
@@ -840,14 +830,12 @@ void UClothingAssetCommon::PostLoad()
 				BackstopDistanceMask.Initialize(BackstopDistances, EWeightMapTargetCommon::BackstopDistance);
 			}
 		}
-#endif
 
 		// Make sure we're transactional
 		SetFlags(RF_Transactional);
 	}
 
 	const int32 ClothingCustomVersion = GetLinkerCustomVersion(FClothingAssetCustomVersion::GUID);
-#if WITH_EDITORONLY_DATA
 	// Fix content imported before we kept vertex colors
 	if (ClothingCustomVersion < FClothingAssetCustomVersion::AddVertexColorsToPhysicalMesh)
 	{
@@ -863,15 +851,12 @@ void UClothingAssetCommon::PostLoad()
 			}
 		}
 	}
-#endif // WITH_EDITORONLY_DATA
 
-#if WITH_EDITOR
+	EClothingCachedDataFlagsCommon ClothingCachedDataFlags = EClothingCachedDataFlagsCommon::None;
 	if (AnimPhysCustomVersion < FAnimPhysObjectVersion::CacheClothMeshInfluences)
 	{
-		// Rebuild data cache
-		InvalidateCachedData();
+		ClothingCachedDataFlags |= EClothingCachedDataFlagsCommon::NumInfluences;
 	}
-#endif
 
 	// Add any missing configs for the available cloth factories, and try to migrate them from any existing one
 	// TODO: Remove all UObject PostLoad dependencies.
@@ -883,7 +868,7 @@ void UClothingAssetCommon::PostLoad()
 			ClothConfig.Value->ConditionalPostLoad();  // PostLoad configs before adding new ones
 		}
 	}
-#if WITH_EDITORONLY_DATA
+
 	if (ClothSimConfig_DEPRECATED)
 	{
 		ClothSimConfig_DEPRECATED->ConditionalPostLoad();  // PostLoad old configs before replacing them
@@ -899,13 +884,16 @@ void UClothingAssetCommon::PostLoad()
 		ClothSharedSimConfig_DEPRECATED->ConditionalPostLoad();  // PostLoad old configs before replacing them
 		ClothSharedSimConfig_DEPRECATED->Rename(nullptr, nullptr, REN_DoNotDirty | REN_DontCreateRedirectors | REN_ForceNoResetLoaders | REN_NonTransactional);  // Rename the config so that the name doesn't collide with the new config map name
 	}
-#endif
-	AddClothConfigs();
+
+	if (AddClothConfigs())
+	{
+		// With a new config added best to recache everything
+		ClothingCachedDataFlags |= EClothingCachedDataFlagsCommon::All;
+	}
 
 	// Migrate configs
 	bool bMigrateSharedConfigToConfig = true;  // Shared config to config migration can be disabled to avoid overriding the newly migrated values
 
-#if WITH_EDITORONLY_DATA
 	if (ClothingCustomVersion < FClothingAssetCustomVersion::MovePropertiesToCommonBaseClasses)
 	{
 		// Remap legacy struct FClothConfig to new config objects
@@ -958,15 +946,26 @@ void UClothingAssetCommon::PostLoad()
 			bMigrateSharedConfigToConfig = false;
 		}
 	}
-#endif
 
 	// Propagate shared configs between cloth assets
 	PropagateSharedConfigs(bMigrateSharedConfigToConfig);
 
+	// Cache tethers and reference bone index should already have been calculated
+	const int32 FrostyMainBranchObjectVersion = GetLinkerCustomVersion(FUE5PrivateFrostyStreamObjectVersion::GUID);
+	if (FrostyMainBranchObjectVersion < FUE5PrivateFrostyStreamObjectVersion::ChaosClothAddTethersToCachedData)
+	{
+		ClothingCachedDataFlags |= EClothingCachedDataFlagsCommon::Tethers;
+
+		// ReferenceBoneIndex is only required when rebinding the cloth.
+		CalculateReferenceBoneIndex();
+	}
+
 	// After fixing the content, we are ready to call functions that rely on it
-	BuildSelfCollisionData();
-#if WITH_EDITORONLY_DATA
-	CalculateReferenceBoneIndex();
+	if (ClothingCachedDataFlags != EClothingCachedDataFlagsCommon::None)
+	{
+		// Rebuild data cache
+		InvalidateFlaggedCachedData(ClothingCachedDataFlags);
+	}
 
 	const int32 PhysicsObjectVersion = GetLinkerCustomVersion(FPhysicsObjectVersion::GUID);
 	const int32 FortniteMainBranchObjectVersion = GetLinkerCustomVersion(FFortniteMainBranchObjectVersion::GUID);
@@ -975,7 +974,7 @@ void UClothingAssetCommon::PostLoad()
 	{
 		BuildLodTransitionData();
 	}
-#endif
+#endif  // #if WITH_EDITORONLY_DATA
 }
 
 void UClothingAssetCommon::Serialize(FArchive& Ar)
@@ -985,10 +984,13 @@ void UClothingAssetCommon::Serialize(FArchive& Ar)
 	Ar.UsingCustomVersion(FClothingAssetCustomVersion::GUID);
 	Ar.UsingCustomVersion(FFortniteMainBranchObjectVersion::GUID);
 	Ar.UsingCustomVersion(FPhysicsObjectVersion::GUID);
+	Ar.UsingCustomVersion(FUE5PrivateFrostyStreamObjectVersion::GUID);
 }
 
-void UClothingAssetCommon::AddClothConfigs()
+bool UClothingAssetCommon::AddClothConfigs()
 {
+	bool bNewConfigAdded = false;
+
 	const TArray<IClothingSimulationFactoryClassProvider*> ClassProviders =
 		IModularFeatures::Get().GetModularFeatureImplementations<IClothingSimulationFactoryClassProvider>(IClothingSimulationFactoryClassProvider::FeatureName);
 
@@ -1029,10 +1031,12 @@ void UClothingAssetCommon::AddClothConfigs()
 					// Add the new config
 					check(ClothConfig);
 					ClothConfigs.Add(ClothConfigName, ClothConfig);
+					bNewConfigAdded = true;
 				}
 			}
 		}
 	}
+	return bNewConfigAdded;
 }
 
 void UClothingAssetCommon::PropagateSharedConfigs(bool bMigrateSharedConfigToConfig)
@@ -1116,113 +1120,105 @@ void UClothingAssetCommon::PropagateSharedConfigs(bool bMigrateSharedConfigToCon
 void UClothingAssetCommon::PostUpdateAllAssets()
 {
 	// Add any missing configs for the available cloth factories, and try to migrate them from any existing one
-	AddClothConfigs();
+	const bool bInvalidateCachedData = AddClothConfigs();
 
 	// Propagate shared configs
 	PropagateSharedConfigs();
+
+#if WITH_EDITORONLY_DATA
+	// Invalidate cached data if the configs have changed
+	if (bInvalidateCachedData)
+	{
+		InvalidateAllCachedData();
+	}
+#endif
 }
 
-#if WITH_EDITOR
 
-void UClothingAssetCommon::InvalidateCachedData()
+#if WITH_EDITORONLY_DATA
+void UClothingAssetCommon::InvalidateFlaggedCachedData(EClothingCachedDataFlagsCommon Flags)
 {
-#if WITH_CHAOS_CLOTHING
-	ForEachInteractorUsingClothing([](UClothingSimulationInteractor* InInteractor)
+	const bool bNeedsInverseMasses = EnumHasAllFlags((EClothingCachedDataFlagsCommon)Flags, EClothingCachedDataFlagsCommon::InverseMasses) && AnyOfClothConfigs([](UClothConfigBase& Config) { return Config.NeedsInverseMasses(); });
+	const bool bNeedsNumInfluences = EnumHasAllFlags((EClothingCachedDataFlagsCommon)Flags, EClothingCachedDataFlagsCommon::NumInfluences) && AnyOfClothConfigs([](UClothConfigBase& Config) { return Config.NeedsNumInfluences(); });
+	const bool bNeedsSelfCollisionData = EnumHasAllFlags((EClothingCachedDataFlagsCommon)Flags, EClothingCachedDataFlagsCommon::SelfCollisionData) && AnyOfClothConfigs([](UClothConfigBase& Config) { return Config.NeedsSelfCollisionData(); });
+	const bool bNeedsTethers = EnumHasAllFlags((EClothingCachedDataFlagsCommon)Flags, EClothingCachedDataFlagsCommon::Tethers) && AnyOfClothConfigs([](UClothConfigBase& Config) { return Config.NeedsTethers(); });
+
+	float SelfCollisionRadius = 0.f;
+	if (bNeedsSelfCollisionData)
+	{
+		// Note: Only PhysX based NvCloth needs to build the SelfCollisionIndices at the moment
+		for (const TPair<FName, TObjectPtr<UClothConfigBase>>& ClothConfig : ClothConfigs)
+		{
+			SelfCollisionRadius = FMath::Max(SelfCollisionRadius, ClothConfig.Value->GetSelfCollisionRadius());
+		}
+	}
+
+	bool bThethersUseEuclideanDistance = false;
+	bool bThethersUseGeodesicDistance = false;
+	if (bNeedsTethers)
 	{
-		if (InInteractor)
+		for (const TPair<FName, TObjectPtr<UClothConfigBase>>& ClothConfig : ClothConfigs)
 		{
-			InInteractor->ClothConfigUpdated();
+			if (ClothConfig.Value->NeedsTethers())
+			{
+				if (ClothConfig.Value->TethersUseGeodesicDistance())
+				{
+					bThethersUseGeodesicDistance = true;
+				}
+				else
+				{
+					bThethersUseEuclideanDistance = true;
+				}
+			}
 		}
-	});
-#endif  // #if WITH_CHAOS_CLOTHING
+	}
 
-	// TODO: This mass calculation isn't used by Chaos, check what can be stripped out
-	for(FClothLODDataCommon& CurrentLodData : LodData)
+	// Recalculate cached data
+	bool bHasClothChanged = false;
+	for (FClothLODDataCommon& CurrentLodData : LodData)
 	{
-		// Recalculate inverse masses for the physical mesh particles
 		FClothPhysicalMeshData& PhysMesh = CurrentLodData.PhysicalMeshData;
-		check(PhysMesh.Indices.Num() % 3 == 0);
-
-		TArray<float>& InvMasses = PhysMesh.InverseMasses;
-
-		const int32 NumVerts = PhysMesh.Vertices.Num();
-		InvMasses.Empty(NumVerts);
-		InvMasses.AddZeroed(NumVerts);
 
-		for(int32 TriBaseIndex = 0; TriBaseIndex < PhysMesh.Indices.Num(); TriBaseIndex += 3)
+		if (bNeedsInverseMasses)
 		{
-			const int32 Index0 = PhysMesh.Indices[TriBaseIndex];
-			const int32 Index1 = PhysMesh.Indices[TriBaseIndex + 1];
-			const int32 Index2 = PhysMesh.Indices[TriBaseIndex + 2];
-
-			const FVector AB = PhysMesh.Vertices[Index1] - PhysMesh.Vertices[Index0];
-			const FVector AC = PhysMesh.Vertices[Index2] - PhysMesh.Vertices[Index0];
-			const float TriArea = FVector::CrossProduct(AB, AC).Size();
-
-			InvMasses[Index0] += TriArea;
-			InvMasses[Index1] += TriArea;
-			InvMasses[Index2] += TriArea;
+			PhysMesh.CalculateInverseMasses();
+			bHasClothChanged = true;
 		}
 
-		PhysMesh.NumFixedVerts = 0;
-
-		const FPointWeightMap* const MaxDistances = PhysMesh.FindWeightMap(EWeightMapTargetCommon::MaxDistance);
-		const TFunction<bool(int32)> IsKinematic = (!MaxDistances || !MaxDistances->Num()) ?
-			TFunction<bool(int32)>([](int32)->bool { return false; }) :
-			TFunction<bool(int32)>([&MaxDistances](int32 Index)->bool { return (*MaxDistances)[Index] < SMALL_NUMBER; });  // For consistency, the default Threshold should be 0.1, not SMALL_NUMBER. But for backward compatibility it needs to be SMALL_NUMBER for now.
-
-		float MassSum = 0.0f;
-		for (int32 CurrVertIndex = 0; CurrVertIndex < NumVerts; ++CurrVertIndex)
+		if (bNeedsNumInfluences)
 		{
-			float& InvMass = InvMasses[CurrVertIndex];
-
-			if (IsKinematic(CurrVertIndex))
-			{
-				InvMass = 0.0f;
-				++PhysMesh.NumFixedVerts;
-			}
-			else
-			{
-				MassSum += InvMass;
-			}
+			PhysMesh.CalculateNumInfluences();
+			bHasClothChanged = true;
 		}
 
-		if (MassSum > 0.0f)
+		if (bNeedsSelfCollisionData)
 		{
-			const float MassScale = (float)(NumVerts - PhysMesh.NumFixedVerts) / MassSum;
-			for (float& InvMass : InvMasses)
-			{
-				if (InvMass != 0.0f)
-				{
-					InvMass *= MassScale;
-					InvMass = 1.0f / InvMass;
-				}
-			}
+			PhysMesh.BuildSelfCollisionData(SelfCollisionRadius);
+			bHasClothChanged = true;
 		}
 
-		// Calculate number of influences per vertex
-		for(int32 VertIndex = 0; VertIndex < NumVerts; ++VertIndex)
+		if (bNeedsTethers)
 		{
-			FClothVertBoneData& BoneData = PhysMesh.BoneData[VertIndex];
-			const uint16* BoneIndices = BoneData.BoneIndices;
-			const float* BoneWeights = BoneData.BoneWeights;
-
-			BoneData.NumInfluences = MAX_TOTAL_INFLUENCES;
+			PhysMesh.CalculateTethers(bThethersUseEuclideanDistance, bThethersUseGeodesicDistance);
+			bHasClothChanged = true;
+		}
+	}
 
-			int32 NumInfluences = 0;
-			for(int32 InfluenceIndex = 0; InfluenceIndex < MAX_TOTAL_INFLUENCES; ++InfluenceIndex)
+	// Inform the simulations that the cloth has changed
+	if (bHasClothChanged)
+	{
+		ForEachInteractorUsingClothing([](UClothingSimulationInteractor* InInteractor)
+		{
+			if (InInteractor)
 			{
-				if(BoneWeights[InfluenceIndex] == 0.0f || BoneIndices[InfluenceIndex] == INDEX_NONE)
-				{
-					BoneData.NumInfluences = NumInfluences;
-					break;
-				}
-				++NumInfluences;
+				InInteractor->ClothConfigUpdated();
 			}
-		}
+		});
 	}
 }
+#endif // #if WITH_EDITORONLY_DATA
 
+#if WITH_EDITOR
 int32 UClothingAssetCommon::AddNewLod()
 {
 	return LodData.AddDefaulted();
@@ -1236,10 +1232,18 @@ void UClothingAssetCommon::PostEditChangeChainProperty(FPropertyChangedChainEven
 
 	if (ChainEvent.ChangeType != EPropertyChangeType::Interactive)
 	{
-		if (ChainEvent.Property->GetFName() == FName("SelfCollisionRadius") ||
-			ChainEvent.Property->GetFName() == FName("SelfCollisionCullScale"))
+		const FName& PropertyName = ChainEvent.PropertyChain.GetActiveMemberNode() && ChainEvent.PropertyChain.GetActiveMemberNode()->GetNextNode() ?
+			ChainEvent.PropertyChain.GetActiveMemberNode()->GetNextNode()->GetValue()->GetFName() : NAME_None;
+
+		if (PropertyName == FName("SelfCollisionRadius") ||
+			PropertyName == FName("SelfCollisionCullScale"))
+		{
+			InvalidateFlaggedCachedData(EClothingCachedDataFlagsCommon::SelfCollisionData);
+			bReregisterComponents = true;
+		}
+		else if (PropertyName == FName("bUseGeodesicDistance"))
 		{
-			BuildSelfCollisionData();
+			InvalidateFlaggedCachedData(EClothingCachedDataFlagsCommon::Tethers);
 			bReregisterComponents = true;
 		}
 		else if(ChainEvent.Property->GetFName() == GET_MEMBER_NAME_CHECKED(UClothingAssetCommon, PhysicsAsset))
diff --git a/Engine/Source/Runtime/ClothingSystemRuntimeCommon/Public/ClothConfig.h b/Engine/Source/Runtime/ClothingSystemRuntimeCommon/Public/ClothConfig.h
--- a/Engine/Source/Runtime/ClothingSystemRuntimeCommon/Public/ClothConfig.h
+++ b/Engine/Source/Runtime/ClothingSystemRuntimeCommon/Public/ClothConfig.h
@@ -40,6 +40,26 @@ public:
 	 * @return true when the migration is possible, false otherwise.
 	 */
 	virtual bool MigrateTo(FClothConfig_Legacy&) const { return false; }
+
+	//~ Begin UClothConfigBase Interface
+	/** Return whether to pre-compute self collision data. */
+	virtual bool NeedsSelfCollisionData() const override { return false; }
+
+	/** Return whether to pre-compute inverse masses. */
+	virtual bool NeedsInverseMasses() const override { return false; }
+
+	/** Return whether to pre-compute the influences. */
+	virtual bool NeedsNumInfluences() const override { return true; }
+
+	/** Return whether to pre-compute the long range attachment tethers. */
+	virtual bool NeedsTethers() const override { return false; }
+
+	/** Return the self collision radius to precomute self collision data. */
+	virtual float GetSelfCollisionRadius() const override { return 0.f; }
+
+	/** Return whether tethers need to be calculated using geodesic distances instead of eclidean. */
+	virtual bool TethersUseGeodesicDistance() const override { return false; }
+	//~ End UClothConfigBase Interface
 };
 
 /** Common shared configuration base class. */
diff --git a/Engine/Source/Runtime/ClothingSystemRuntimeCommon/Public/ClothPhysicalMeshData.h b/Engine/Source/Runtime/ClothingSystemRuntimeCommon/Public/ClothPhysicalMeshData.h
--- a/Engine/Source/Runtime/ClothingSystemRuntimeCommon/Public/ClothPhysicalMeshData.h
+++ b/Engine/Source/Runtime/ClothingSystemRuntimeCommon/Public/ClothPhysicalMeshData.h
@@ -2,6 +2,7 @@
 #pragma once
 
 #include "ClothVertBoneData.h"
+#include "ClothTetherData.h"
 #include "PointWeightMap.h"
 #include "Containers/Array.h"
 #include "Containers/Map.h"
@@ -33,42 +34,48 @@ struct CLOTHINGSYSTEMRUNTIMECOMMON_API FClothPhysicalMeshData
 	void ClearWeightMaps();
 
 	/** Build the self collision indices for the relevant config. */
+	UE_DEPRECATED(5.0, "Use BuildSelfCollisionData(float SelfCollisionRadius) instead.")
 	void BuildSelfCollisionData(const TMap<FName, TObjectPtr<UClothConfigBase>>& ClothConfigs);
 
+	/** Build the self collision indices with the specified radius. */
+	void BuildSelfCollisionData(float SelfCollisionRadius);
+
+	/** Recalculate the node inverse masses. */
+	void CalculateInverseMasses();
+
+	/** Recalculate the number of influences for the bone data. */
+	void CalculateNumInfluences();
+
+	/** Recalculate the long range attachment tethers. */
+	void CalculateTethers(bool bUseEuclideanDistance, bool bUseGeodesicDistance);
+
 	/** Retrieve whether a vertex weight array has already been registered. */
 	template<typename T>
-	bool HasWeightMap(const T Target) const
-	{ return WeightMaps.Contains((uint32)Target); }
+	bool HasWeightMap(const T Target) const { return WeightMaps.Contains((uint32)Target); }
 
 	/** Retrieve a pointer to a registered vertex weight array by unique @param Id, or nullptr if none is found. */
 	template<typename T>
-	const FPointWeightMap* FindWeightMap(const T Target) const
-	{ return WeightMaps.Find((uint32)Target); }
+	const FPointWeightMap* FindWeightMap(const T Target) const { return WeightMaps.Find((uint32)Target); }
 
 	/** Retrieve a pointer to a registered vertex weight array by unique @param Id, or nullptr if none is found. */
 	template<typename T>
-	FPointWeightMap* FindWeightMap(const T Target)
-	{ return WeightMaps.Find((uint32)Target); }
+	FPointWeightMap* FindWeightMap(const T Target) { return WeightMaps.Find((uint32)Target); }
 
 	/** Retrieve a pointer to a registered vertex weight array by unique @param Id, or add one if it doesn't exist already. */
 	template<typename T>
-	FPointWeightMap& AddWeightMap(const T Target)
-	{ return WeightMaps.Add((uint32)Target); }
+	FPointWeightMap& AddWeightMap(const T Target) { return WeightMaps.Add((uint32)Target); }
 
 	/** Retrieve a pointer to a registered vertex weight array by unique @param Id, or add one if it doesn't exist already. */
 	template<typename T>
-	FPointWeightMap& FindOrAddWeightMap(const T Target)
-	{ return WeightMaps.FindOrAdd((uint32)Target); }
+	FPointWeightMap& FindOrAddWeightMap(const T Target) { return WeightMaps.FindOrAdd((uint32)Target); }
 
 	/** Retrieve a registered vertex weight array by unique @param Id. The array must exists or this function will assert. */
 	template<typename T>
-	const FPointWeightMap& GetWeightMap(const T Target) const
-	{ return WeightMaps[(uint32)Target]; }
+	const FPointWeightMap& GetWeightMap(const T Target) const { return WeightMaps[(uint32)Target]; }
 
 	/** Retrieve a registered vertex weight array by unique @param Id. The array must exists or this function will assert. */
 	template<typename T>
-	FPointWeightMap& GetWeightMap(const T Target)
-	{ return WeightMaps[(uint32)Target]; }
+	FPointWeightMap& GetWeightMap(const T Target) { return WeightMaps[(uint32)Target]; }
 
 	// Positions of each simulation vertex
 	UPROPERTY(EditAnywhere, Category = SimMesh)
@@ -100,6 +107,18 @@ struct CLOTHINGSYSTEMRUNTIMECOMMON_API FClothPhysicalMeshData
 	UPROPERTY(EditAnywhere, Category = SimMesh)
 	TArray<FClothVertBoneData> BoneData;
 
+	// Valid indices to use for self collisions (reduced set of Indices)
+	UPROPERTY(EditAnywhere, Category = SimMesh)
+	TArray<uint32> SelfCollisionIndices;
+
+	// Long range attachment tethers, using euclidean (beeline) distance to find the closest attachment
+	UPROPERTY(EditAnywhere, Category = SimMesh)
+	FClothTetherData EuclideanTethers;
+
+	// Long range attachment tethers, using geodesic (surface) distance to find the closest attachment
+	UPROPERTY(EditAnywhere, Category = SimMesh)
+	FClothTetherData GeodesicTethers;
+
 	// Maximum number of bone weights of any vetex
 	UPROPERTY(EditAnywhere, Category = SimMesh)
 	int32 MaxBoneWeights;
@@ -108,10 +127,7 @@ struct CLOTHINGSYSTEMRUNTIMECOMMON_API FClothPhysicalMeshData
 	UPROPERTY(EditAnywhere, Category = SimMesh)
 	int32 NumFixedVerts;
 
-	// Valid indices to use for self collisions (reduced set of Indices)
-	UPROPERTY(EditAnywhere, Category = SimMesh)
-	TArray<uint32> SelfCollisionIndices;
-
+#if WITH_EDITORONLY_DATA
 	// Deprecated. Use WeightMaps instead.
 	UPROPERTY()
 	TArray<float> MaxDistances_DEPRECATED;
@@ -121,4 +137,5 @@ struct CLOTHINGSYSTEMRUNTIMECOMMON_API FClothPhysicalMeshData
 	TArray<float> BackstopRadiuses_DEPRECATED;
 	UPROPERTY()
 	TArray<float> AnimDriveMultipliers_DEPRECATED;
+#endif // WITH_EDITORONLY_DATA
 };
diff --git a/Engine/Source/Runtime/ClothingSystemRuntimeCommon/Public/ClothTetherData.h b/Engine/Source/Runtime/ClothingSystemRuntimeCommon/Public/ClothTetherData.h
new file mode 100644
--- /dev/null
+++ b/Engine/Source/Runtime/ClothingSystemRuntimeCommon/Public/ClothTetherData.h
@@ -0,0 +1,42 @@
+// Copyright Epic Games, Inc. All Rights Reserved.
+#pragma once
+
+#include "Containers/ContainersFwd.h"
+
+#include "ClothTetherData.generated.h"
+
+/**
+ * Long range attachment tether pathfinding based on Dijkstra's algorithm.
+ * For more information about the long range attachment, see:
+ * https://matthias-research.github.io/pages/publications/sca2012cloth.pdf
+ */
+USTRUCT()
+struct FClothTetherData
+{
+	GENERATED_BODY()
+
+	/** Long range attachment tether start/end/length, sorted in sequential batches of independant tethers. */
+	TArray<TArray<TTuple<int32, int32, float>>> Tethers;
+	
+	/**
+	 * Generate the tethers by following the triangle mesh network from the closest kinematic to each dynamic point.
+	 * Inside array items can be processed concurrently, but the outside array must be iterated on sequentially.
+	 */
+	void GenerateTethers(
+		const TConstArrayView<FVector3f>& Points,  // Reference pose
+		const TConstArrayView<uint32>& Indices,  // Triangle mesh
+		const TConstArrayView<float>& MaxDistances,  // Mask for sorting the kinematic from the dynamic points
+		bool bUseGeodesicDistance);  // Whether to use geodesic (walking along the sruface) or euclidean (beeline) distances to find the tethers.
+
+	/** Custom serializer, since neiher an array of array nor a tuple can be set as a UPROPERTY. */ 
+	bool Serialize(FArchive& Ar);
+};
+
+template<>
+struct TStructOpsTypeTraits<FClothTetherData> : public TStructOpsTypeTraitsBase2<FClothTetherData>
+{
+	enum
+	{
+		WithSerializer = true,
+	};
+};
diff --git a/Engine/Source/Runtime/ClothingSystemRuntimeCommon/Public/ClothingAsset.h b/Engine/Source/Runtime/ClothingSystemRuntimeCommon/Public/ClothingAsset.h
--- a/Engine/Source/Runtime/ClothingSystemRuntimeCommon/Public/ClothingAsset.h
+++ b/Engine/Source/Runtime/ClothingSystemRuntimeCommon/Public/ClothingAsset.h
@@ -7,7 +7,7 @@
 #include "ClothConfig_Legacy.h"
 #include "ClothLODData_Legacy.h"
 #include "ClothingSimulationInteractor.h"
-
+#include "Algo/AnyOf.h"
 #include "ClothingAsset.generated.h"
 
 DECLARE_LOG_CATEGORY_EXTERN(LogClothingAsset, Log, All);
@@ -60,16 +60,9 @@ namespace ClothingAssetUtils
 
 /**
  * Custom data wrapper for clothing assets.
- * If writing a new clothing asset importer, creating a new derived custom data 
- * is how to store importer (and possibly simulation) data that importer will 
- * create. This needs to be set to the \c CustomData member on the asset your 
- * factory creates.
- *
- * Testing whether a UClothingAssetCommon was made from a custom plugin can be 
- * achieved with:
- * \code if(AssetPtr->CustomData->IsA(UMyCustomData::StaticClass())) \endcode
  */
-UCLASS(abstract, MinimalAPI)
+class UE_DEPRECATED(5.0, "Redundant class.") UClothingAssetCustomData;
+UCLASS(Abstract, MinimalAPI)
 class UClothingAssetCustomData : public UObject
 {
 	GENERATED_BODY()
@@ -78,6 +71,20 @@ public:
 	{}
 };
 
+/**
+ * Common flags used by InvalidateCachedData.
+ */
+enum class EClothingCachedDataFlagsCommon : uint8
+{
+	None = 0,
+	InverseMasses = 1 << 0,
+	NumInfluences = 1 << 1,
+	SelfCollisionData = 1 << 2,
+	Tethers = 1 << 3,
+	All = 0xFF
+};
+ENUM_CLASS_FLAGS(EClothingCachedDataFlagsCommon);
+
 /**
  * Implementation of non-solver specific, but common Engine related functionality.
  *
@@ -140,15 +147,23 @@ public:
 	/** Returns the number of valid LOD's (length of the \c ClothLodData array). */
 	virtual int32 GetNumLods() const override;
 
-#if WITH_EDITOR
-	/** * Add a new LOD class instance. */
-	virtual int32 AddNewLod() override;
+#if WITH_EDITORONLY_DATA
+	/**
+	 * Called on the clothing asset when the base data (physical mesh, config etc.)
+	 * has changed, so any intermediate generated data can be regenerated.
+	 */
+	virtual void InvalidateAllCachedData() override { InvalidateFlaggedCachedData(EClothingCachedDataFlagsCommon::All); }
 
 	/**
-	 * Called on the clothing asset when the base data (physical mesh etc.) has
-	 * changed, so any intermediate generated data can be regenerated.
+	 * Called on the clothing asset when the base data (physical mesh, config etc.)
+	 * has changed, so any intermediate generated data can be regenerated.
 	 */
-	virtual void InvalidateCachedData() override;
+	void InvalidateFlaggedCachedData(EClothingCachedDataFlagsCommon Flags);
+#endif // WITH_EDITORONLY_DATA
+
+#if WITH_EDITOR
+	/** * Add a new LOD class instance. */
+	virtual int32 AddNewLod() override;
 
 	/* Called after changes in any of the asset properties. */
 	virtual void PostEditChangeChainProperty(FPropertyChangedChainEvent& ChainEvent) override;
@@ -160,7 +175,13 @@ public:
 	 * are enabled, it envokes \c BuildSelfCollisionData() on each \c ClothLodData's
 	 * \c ClothPhysicalMeshData member.
 	 */
-	virtual void BuildSelfCollisionData() override;
+	UE_DEPRECATED(5.0, "Cached data are now all rebuilt by calling InvalidateCachedData")
+	virtual void BuildSelfCollisionData() override
+	{
+#if WITH_EDITORONLY_DATA
+		InvalidateFlaggedCachedData(EClothingCachedDataFlagsCommon::SelfCollisionData);
+#endif // WITH_EDITORONLY_DATA
+	}
 
 	/** Return a const cloth config pointer of the desired cloth config type, or nullptr if there isn't any suitable. */
 	template<typename ClothConfigType, typename = typename TEnableIf<TIsDerivedFrom<ClothConfigType, UClothConfigBase>::IsDerived>::Type>
@@ -240,9 +261,11 @@ public:
 	UPROPERTY()
 	int32 ReferenceBoneIndex;
 
+PRAGMA_DISABLE_DEPRECATION_WARNINGS
 	// Custom data applied by the importer depending on where the asset was imported from.
 	UPROPERTY()
 	TObjectPtr<UClothingAssetCustomData> CustomData;
+PRAGMA_ENABLE_DEPRECATION_WARNINGS
 
 #if WITH_EDITORONLY_DATA
 	/** 
@@ -266,13 +289,25 @@ private:
 	// Create and add any missing cloth configs.
 	// If a config from a different factory exists already, the newly
 	// created config will attempt to initialize its parameters from it.
-	void AddClothConfigs();
+	// Return true when at least one config has been added, false otherwise.
+	bool AddClothConfigs();
 
 	// Propagate the shared simulation configs between assets.
 	// Also migrate all deprecated shared parameters which have been moved to the per cloth configs if required.
 	// Called after a cloth asset is created or loaded.
 	void PropagateSharedConfigs(bool bMigrateSharedConfigToConfig=false);
 
+	// Return true when any one of the cloth configs fullfill the predicate.
+	// Used to select which type of data to cache.
+	template<typename PredicateType>
+	bool AnyOfClothConfigs(PredicateType Predicate) const
+	{
+		return Algo::AnyOf(ClothConfigs, [&Predicate](const TPair<FName, TObjectPtr<UClothConfigBase>>& ClothConfig)
+			{
+				return Predicate(*ClothConfig.Value);
+			});
+	}
+
 #if WITH_EDITOR
 	// Helper functions used in PostPropertyChangeCb
 	void ReregisterComponentsUsingClothing();
diff --git a/Engine/Source/Runtime/ClothingSystemRuntimeInterface/Public/ClothConfigBase.h b/Engine/Source/Runtime/ClothingSystemRuntimeInterface/Public/ClothConfigBase.h
--- a/Engine/Source/Runtime/ClothingSystemRuntimeInterface/Public/ClothConfigBase.h
+++ b/Engine/Source/Runtime/ClothingSystemRuntimeInterface/Public/ClothConfigBase.h
@@ -17,19 +17,33 @@ public:
 	UClothConfigBase();
 	virtual ~UClothConfigBase();
 
-	/**
-	 * Return whether self collision is enabled for this config.
-	 */
-	UE_DEPRECATED(4.25, "This function is deprecated. Please use NeedsSelfCollisionIndices instead.")
-	virtual bool HasSelfCollision() const
-	{ unimplemented(); return false; }
+	/** Return the self collision radius if building self collision indices is required for this config, or 0.f otherwise. */
+	UE_DEPRECATED(5.0, "Use NeedsSelfCollisionData and GetSelfCollisionRadius instead.")
+	virtual float NeedsSelfCollisionIndices() const { return NeedsSelfCollisionData() ? GetSelfCollisionRadius() : 0.f; }
 
-	/**
-	 * Return the self collision radius if building self collision indices is required for this config.
-	 * Otherwise return 0.f.
-	 */
-	virtual float NeedsSelfCollisionIndices() const
-	{ return 0.f; }
+	/** Return wherether to pre-compute self collision data. */
+	virtual bool NeedsSelfCollisionData() const
+	PURE_VIRTUAL(UClothConfigBase::NeedsSelfCollisionData, return false;);
+
+	/** Return wherether to pre-compute inverse masses. */
+	virtual bool NeedsInverseMasses() const
+	PURE_VIRTUAL(UClothConfigBase::NeedsInverseMasses, return false;);
+
+	/** Return wherether to pre-compute the influences. */
+	virtual bool NeedsNumInfluences() const
+	PURE_VIRTUAL(UClothConfigBase::NeedsNumInfluences, return false;);
+
+	/** Return wherether to pre-compute the long range attachment tethers. */
+	virtual bool NeedsTethers() const
+	PURE_VIRTUAL(UClothConfigBase::NeedsTethers, return false;);
+
+	/** Return the self collision radius to precomute self collision data. */
+	virtual float GetSelfCollisionRadius() const
+	PURE_VIRTUAL(UClothConfigBase::GetSelfCollisionRadius, return 0.f;);
+
+	/** Return whether tethers need to be calculated using geodesic distances instead of eclidean. */
+	virtual bool TethersUseGeodesicDistance() const
+	PURE_VIRTUAL(UClothConfigBase::TethersUseGeodesicDistance, return false;);
 };
 
 /**
diff --git a/Engine/Source/Runtime/ClothingSystemRuntimeInterface/Public/ClothingAssetBase.h b/Engine/Source/Runtime/ClothingSystemRuntimeInterface/Public/ClothingAssetBase.h
--- a/Engine/Source/Runtime/ClothingSystemRuntimeInterface/Public/ClothingAssetBase.h
+++ b/Engine/Source/Runtime/ClothingSystemRuntimeInterface/Public/ClothingAssetBase.h
@@ -8,8 +8,10 @@
 
 class USkeletalMesh;
 
-// An interface object for any clothing asset the engine can use. Any
-// clothing asset concrete object should derive from this.
+/**
+ * An interface object for any clothing asset the engine can use.
+ * Any clothing asset concrete object should derive from this.
+ */
 UCLASS(Abstract)
 class CLOTHINGSYSTEMRUNTIMEINTERFACE_API UClothingAssetBase : public UObject
 {
@@ -46,8 +48,13 @@ public:
 	 * Called on the clothing asset when the base data (physical mesh etc.) has
 	 * changed, so any intermediate generated data can be regenerated.
 	 */
+	UE_DEPRECATED(5.0, "Use InvalidateAllCachedData() instead")
 	virtual void InvalidateCachedData()
-	PURE_VIRTUAL(UClothingAssetBase::InvalidateCachedData(), );
+	{
+#if WITH_EDITORONLY_DATA
+		InvalidateAllCachedData();
+#endif
+	}
 
 	/** Add a new LOD class instance. */
 	virtual int32 AddNewLod()
@@ -62,6 +69,15 @@ public:
 	PURE_VIRTUAL(UClothingAssetBase::BuildLodTransitionData(), );
 #endif
 
+#if WITH_EDITORONLY_DATA
+	/**
+	 * Called on the clothing asset when the base data (physical mesh, config etc.)
+	 * has changed, so any intermediate generated data can be regenerated.
+	 */
+	virtual void InvalidateAllCachedData()
+	PURE_VIRTUAL(UClothingAssetBase::InvalidateCachedData(),);
+#endif // WITH_EDITORONLY_DATA
+
 	/** 
 	 * Messages to the clothing asset that the bones in the parent mesh have
 	 * possibly changed, which could invalidate the bone indices stored in the LOD
@@ -80,6 +96,7 @@ public:
 	PURE_VIRTUAL(UClothingAssetBase::GetNumLods(), return 0;);
 
 	/** Builds self collision data */
+	UE_DEPRECATED(5.0, "Cached data are now all rebuilt by calling InvalidateCachedData()")
 	virtual void BuildSelfCollisionData()
 	PURE_VIRTUAL(UClothingAssetBase::BuildSelfCollisionData(), );
 
diff --git a/Engine/Source/Runtime/ClothingSystemRuntimeNv/Private/ClothConfigNv.cpp b/Engine/Source/Runtime/ClothingSystemRuntimeNv/Private/ClothConfigNv.cpp
--- a/Engine/Source/Runtime/ClothingSystemRuntimeNv/Private/ClothConfigNv.cpp
+++ b/Engine/Source/Runtime/ClothingSystemRuntimeNv/Private/ClothConfigNv.cpp
@@ -138,14 +138,3 @@ bool UClothConfigNv::MigrateTo(FClothConfig_Legacy& ClothConfig) const
 
 	return true;
 }
-
-
-bool UClothConfigNv::UseSelfCollisions() const
-{
-	return SelfCollisionRadius > 0.0f && SelfCollisionStiffness > 0.0f;
-}
-
-float UClothConfigNv::NeedsSelfCollisionIndices() const
-{
-	return UseSelfCollisions() ? SelfCollisionRadius * SelfCollisionCullScale : 0.0f;
-}
diff --git a/Engine/Source/Runtime/ClothingSystemRuntimeNv/Public/ClothConfigNv.h b/Engine/Source/Runtime/ClothingSystemRuntimeNv/Public/ClothConfigNv.h
--- a/Engine/Source/Runtime/ClothingSystemRuntimeNv/Public/ClothConfigNv.h
+++ b/Engine/Source/Runtime/ClothingSystemRuntimeNv/Public/ClothConfigNv.h
@@ -72,16 +72,21 @@ public:
 	// Migrate to the legacy FClothConfig structure.
 	virtual bool MigrateTo(FClothConfig_Legacy& ClothConfig) const override;
 
-	// Return whether self collision is enabled for this config.
-	UE_DEPRECATED(4.25, "This function is deprecated. Please use NeedsSelfCollisionIndices or UseSelfCollisions instead.")
-	virtual bool HasSelfCollision() const override
-	{ return UseSelfCollisions(); }
+	/** Return whether to pre-compute self collision data. */
+	virtual bool NeedsSelfCollisionData() const override { return SelfCollisionRadius > 0.0f && SelfCollisionStiffness > 0.0f; }
+
+	/** Return whether to pre-compute inverse masses. */
+	virtual bool NeedsInverseMasses() const override { return true; }
+
+	/** Return whether to pre-compute the long range attachment tethers. */
+	virtual bool NeedsTethers() const override { return false; }
 
 	// Return the collision radius required to calculate the self collision indices, or 0.f if self collision is disabled.
-	virtual float NeedsSelfCollisionIndices() const override;
+	virtual float GetSelfCollisionRadius() const override { return NeedsSelfCollisionData() ? SelfCollisionRadius * SelfCollisionCullScale : 0.0f; }
 
 	// Return whether this Nv config has self collision.
-	bool UseSelfCollisions() const;
+	UE_DEPRECATED(5.0, "Use NeedsSelfCollisionData instead.")
+	bool UseSelfCollisions() const { return NeedsSelfCollisionData(); }
 
 	// How wind should be processed, Accurate uses drag and lift to make the cloth react differently, legacy applies similar forces to all clothing without drag and lift (similar to APEX)
 	UPROPERTY(EditAnywhere, Category = ClothConfig)
diff --git a/Engine/Source/Runtime/Core/Public/UObject/UE5PrivateFrostyStreamObjectVersion.h b/Engine/Source/Runtime/Core/Public/UObject/UE5PrivateFrostyStreamObjectVersion.h
--- a/Engine/Source/Runtime/Core/Public/UObject/UE5PrivateFrostyStreamObjectVersion.h
+++ b/Engine/Source/Runtime/Core/Public/UObject/UE5PrivateFrostyStreamObjectVersion.h
@@ -18,6 +18,9 @@ struct CORE_API FUE5PrivateFrostyStreamObjectVersion
 		// Serialize scene components static bounds
 		SerializeSceneComponentStaticBounds,
 
+		// Add the long range attachment tethers to the cloth asset to avoid a large hitch during the cloth's initialization.
+		ChaosClothAddTethersToCachedData,
+
 		// -----<new versions can be added above this line>-------------------------------------------------
 		VersionPlusOne,
 		LatestVersion = VersionPlusOne - 1
diff --git a/Engine/Source/Runtime/Engine/Private/SkeletalMesh.cpp b/Engine/Source/Runtime/Engine/Private/SkeletalMesh.cpp
--- a/Engine/Source/Runtime/Engine/Private/SkeletalMesh.cpp
+++ b/Engine/Source/Runtime/Engine/Private/SkeletalMesh.cpp
@@ -3339,9 +3339,7 @@ void USkeletalMesh::FinishPostLoadInternal(FSkeletalMeshPostLoadContext& Context
 		// Convert screen area to screen size
 		ConvertLegacyLODScreenSize();
 	}
-#endif
 
-#if WITH_EDITOR
 	// If inverse masses have never been cached, invalidate data so it will be recalculated
 	if (GetLinkerCustomVersion(FSkeletalMeshCustomVersion::GUID) < FSkeletalMeshCustomVersion::CachedClothInverseMasses)
 	{
@@ -3349,7 +3347,7 @@ void USkeletalMesh::FinishPostLoadInternal(FSkeletalMeshPostLoadContext& Context
 		{
 			if (ClothingAsset) 
 			{
-				ClothingAsset->InvalidateCachedData();
+				ClothingAsset->InvalidateAllCachedData();
 			}
 		}
 	}
diff --git a/Engine/Source/Runtime/Experimental/Chaos/Private/Chaos/PBDLongRangeConstraints.cpp b/Engine/Source/Runtime/Experimental/Chaos/Private/Chaos/PBDLongRangeConstraints.cpp
--- a/Engine/Source/Runtime/Experimental/Chaos/Private/Chaos/PBDLongRangeConstraints.cpp
+++ b/Engine/Source/Runtime/Experimental/Chaos/Private/Chaos/PBDLongRangeConstraints.cpp
@@ -12,74 +12,157 @@ FAutoConsoleVariableRef CVarChaosLongRangeISPCEnabled(TEXT("p.Chaos.LongRange.IS
 
 using namespace Chaos;
 
-void FPBDLongRangeConstraints::Apply(FPBDParticles& Particles, const FReal /*Dt*/, const TArray<int32>& ConstraintIndices) const
-{
-	TRACE_CPUPROFILER_EVENT_SCOPE(FPBDLongRangeConstraints_Apply);
-	SCOPE_CYCLE_COUNTER(STAT_PBD_LongRange);
-	if (!Stiffness.HasWeightMap())
-	{
-		const FReal ExpStiffnessValue = (FReal)Stiffness;
-		for (const int32 ConstraintIndex : ConstraintIndices)
-		{
-			const FTether& Tether = Tethers[ConstraintIndex];
-			Particles.P(Tether.End) += Stiffness[Tether.End - ParticleOffset] * Tether.GetDelta(Particles);
-		}
-	}
-	else
-	{
-		for (const int32 ConstraintIndex : ConstraintIndices)
-		{
-			const FTether& Tether = Tethers[ConstraintIndex];
-			const FReal ExpStiffnessValue = Stiffness[Tether.End - ParticleOffset];
-			Particles.P(Tether.End) += Stiffness[Tether.End - ParticleOffset] * Tether.GetDelta(Particles);
-		}
-	}
-}
-
 void FPBDLongRangeConstraints::Apply(FPBDParticles& Particles, const FReal /*Dt*/) const
 {
 	TRACE_CPUPROFILER_EVENT_SCOPE(FPBDLongRangeConstraints_Apply);
 	SCOPE_CYCLE_COUNTER(STAT_PBD_LongRange);
 
+	const int32 MinParallelSize = GetMinParallelBatchSize();
+
 	if (!Stiffness.HasWeightMap())
 	{
 		const FReal ExpStiffnessValue = (FReal)Stiffness;
-#if INTEL_ISPC
-		if (bRealTypeCompatibleWithISPC && bChaos_LongRange_ISPC_Enabled)
+		if (!HasScaleWeightMap())
 		{
-			// Run particles in parallel, and ranges in sequence to avoid a race condition when updating the same particle from different tethers
-			TethersView.RangeFor([this, &Particles, ExpStiffnessValue](TArray<FTether>& InTethers, int32 Offset, int32 Range)
+			const FReal ScaleValue = ScaleTable[0];
+#if INTEL_ISPC
+			if (bRealTypeCompatibleWithISPC && bChaos_LongRange_ISPC_Enabled)
+			{
+				// Run particles in parallel, and batch in sequence to avoid a race condition when updating the same particle from different tethers
+				for (const TConstArrayView<FTether>& TetherBatch : Tethers)
 				{
 					ispc::ApplyLongRangeConstraints(
 						(ispc::FVector*)Particles.GetP().GetData(),
-						(ispc::FTether*)(InTethers.GetData() + Offset),
+						(const ispc::FTether*)TetherBatch.GetData(),
 						ExpStiffnessValue,
-						Range - Offset);
-				});
-		}
-		else
+						ScaleValue,
+						TetherBatch.Num(),
+						ParticleOffset);
+				}
+			}
+			else
 #endif
+			{
+				// Run particles in parallel, and batch in sequence to avoid a race condition when updating the same particle from different tethers
+				for (const TConstArrayView<FTether>& TetherBatch : Tethers)
+				{
+					PhysicsParallelFor(TetherBatch.Num(), [this, &Particles, &TetherBatch, ExpStiffnessValue, ScaleValue](int32 Index)
+						{
+							const FTether& Tether = TetherBatch[Index];
+							Particles.P(GetEndParticle(Tether)) += ExpStiffnessValue * GetDelta(Particles, Tether, ScaleValue);
+						}, TetherBatch.Num() < MinParallelSize);
+				}
+			}
+		}
+		else  // HasScaleWeightMap
 		{
-			// Run particles in parallel, and ranges in sequence to avoid a race condition when updating the same particle from different tethers
-			static const int32 MinParallelSize = 500;
-			TethersView.ParallelFor([this, &Particles, ExpStiffnessValue](TArray<FTether>& InTethers, int32 Index)
+#if INTEL_ISPC
+			if (bRealTypeCompatibleWithISPC && bChaos_LongRange_ISPC_Enabled)
+			{
+				// Run particles in parallel, and batch in sequence to avoid a race condition when updating the same particle from different tethers
+				for (const TConstArrayView<FTether>& TetherBatch : Tethers)
+				{
+					ispc::ApplyLongRangeConstraintsScaleWeightmap(
+						(ispc::FVector*)Particles.GetP().GetData(),
+						(const ispc::FTether*)TetherBatch.GetData(),
+						ExpStiffnessValue,
+						ScaleIndices.GetData(),
+						ScaleTable.GetData(),
+						TetherBatch.Num(),
+						ParticleOffset);
+				}
+			}
+			else
+#endif
+			{
+				// Run particles in parallel, and batch in sequence to avoid a race condition when updating the same particle from different tethers
+				for (const TConstArrayView<FTether>& TetherBatch : Tethers)
 				{
-					const FTether& Tether = InTethers[Index];
-					Particles.P(Tether.End) += ExpStiffnessValue * Tether.GetDelta(Particles);
-				}, MinParallelSize);
+					PhysicsParallelFor(TetherBatch.Num(), [this, &Particles, &TetherBatch, ExpStiffnessValue](int32 Index)
+						{
+							const FTether& Tether = TetherBatch[Index];
+							const int32 LocalParticleIndex = GetEndIndex(Tether);
+							const FReal ScaleValue = ScaleTable[ScaleIndices[LocalParticleIndex]];
+							Particles.P(ParticleOffset + LocalParticleIndex) += ExpStiffnessValue * GetDelta(Particles, Tether, ScaleValue);
+						}, TetherBatch.Num() < MinParallelSize);
+				}
+			}
 		}
 	}
-	else
+	else  // HasStiffnessWeighmap
 	{
-		// TODO: ISPC implementation
+		if (!HasScaleWeightMap())
+		{
+			const FReal ScaleValue = ScaleTable[0];
 
-		// Run particles in parallel, and ranges in sequence to avoid a race condition when updating the same particle from different tethers
-		static const int32 MinParallelSize = 500;
-		TethersView.ParallelFor([this, &Particles](TArray<FTether>& InTethers, int32 Index)
+#if INTEL_ISPC
+			if (bRealTypeCompatibleWithISPC && bChaos_LongRange_ISPC_Enabled)
+			{
+				// Run particles in parallel, and batch in sequence to avoid a race condition when updating the same particle from different tethers
+				for (const TConstArrayView<FTether>& TetherBatch : Tethers)
+				{
+					ispc::ApplyLongRangeConstraintsStiffnessWeightmap(
+						(ispc::FVector*)Particles.GetP().GetData(),
+						(const ispc::FTether*)TetherBatch.GetData(),
+						Stiffness.GetIndices().GetData(),
+						Stiffness.GetTable().GetData(),
+						ScaleValue,
+						TetherBatch.Num(),
+						ParticleOffset);
+				}
+			}
+			else
+#endif
 			{
-				const FTether& Tether = InTethers[Index];
-				Particles.P(Tether.End) += Stiffness[Tether.End - ParticleOffset] * Tether.GetDelta(Particles);
-			}, MinParallelSize);
+				// Run particles in parallel, and batch in sequence to avoid a race condition when updating the same particle from different tethers
+				for (const TConstArrayView<FTether>& TetherBatch : Tethers)
+				{
+					PhysicsParallelFor(TetherBatch.Num(), [this, &Particles, &TetherBatch, ScaleValue](int32 Index)
+						{
+							const FTether& Tether = TetherBatch[Index];
+							const int32 LocalParticleIndex = GetEndIndex(Tether);
+							const FReal ExpStiffnessValue = Stiffness[LocalParticleIndex];
+							Particles.P(ParticleOffset + LocalParticleIndex) += ExpStiffnessValue * GetDelta(Particles, Tether, ScaleValue);
+						}, TetherBatch.Num() < MinParallelSize);
+				}
+			}
+		}
+		else // HasScaleWeightMap
+		{
+#if INTEL_ISPC
+			if (bRealTypeCompatibleWithISPC && bChaos_LongRange_ISPC_Enabled)
+			{
+				// Run particles in parallel, and batch in sequence to avoid a race condition when updating the same particle from different tethers
+				for (const TConstArrayView<FTether>& TetherBatch : Tethers)
+				{
+					ispc::ApplyLongRangeConstraintsStiffnessScaleWeightmaps(
+						(ispc::FVector*)Particles.GetP().GetData(),
+						(const ispc::FTether*)TetherBatch.GetData(),
+						Stiffness.GetIndices().GetData(),
+						Stiffness.GetTable().GetData(),
+						ScaleIndices.GetData(),
+						ScaleTable.GetData(),
+						TetherBatch.Num(),
+						ParticleOffset);
+				}
+			}
+			else
+#endif
+			{
+				// Run particles in parallel, and batch in sequence to avoid a race condition when updating the same particle from different tethers
+				for (const TConstArrayView<FTether>& TetherBatch : Tethers)
+				{
+					PhysicsParallelFor(TetherBatch.Num(), [this, &Particles, &TetherBatch](int32 Index)
+						{
+							const FTether& Tether = TetherBatch[Index];
+							const int32 LocalParticleIndex = GetEndIndex(Tether);
+							const FReal ExpStiffnessValue = Stiffness[LocalParticleIndex];
+							const FReal ScaleValue = ScaleTable[ScaleIndices[LocalParticleIndex]];
+							Particles.P(ParticleOffset + LocalParticleIndex) += ExpStiffnessValue * GetDelta(Particles, Tether, ScaleValue);
+						}, TetherBatch.Num() < MinParallelSize);
+				}
+			}
+		}
 	}
 }
 
diff --git a/Engine/Source/Runtime/Experimental/Chaos/Private/Chaos/PBDLongRangeConstraints.ispc b/Engine/Source/Runtime/Experimental/Chaos/Private/Chaos/PBDLongRangeConstraints.ispc
--- a/Engine/Source/Runtime/Experimental/Chaos/Private/Chaos/PBDLongRangeConstraints.ispc
+++ b/Engine/Source/Runtime/Experimental/Chaos/Private/Chaos/PBDLongRangeConstraints.ispc
@@ -24,7 +24,9 @@ static inline varying FTether LoadTetherHelper(const uniform FTether *uniform Sr
 export void ApplyLongRangeConstraints(uniform FVector P[],
 										const uniform FTether Tethers[],
 										const uniform float Stiffness,
-										const uniform int32 NumConstraints)
+										const uniform float Scale,
+										const uniform int32 NumConstraints,
+										const uniform int32 ParticleOffset)
 {
 	varying FTether Tether;
 
@@ -40,9 +42,13 @@ export void ApplyLongRangeConstraints(uniform FVector P[],
 #if HW_GATHER_SUPPORTED == 0
 		Tether = LoadTetherHelper(&Tethers[extract(i, 0)]);
 #endif
-		const varying int32 i1 = Tether.Start;
-		const varying int32 i2 = Tether.End;
+		const varying int32 Start = Tether.Start;
+		const varying int32 End = Tether.End;
 		const varying float RefLength = Tether.RefLength;
+
+		const varying int32 i1 = Start + ParticleOffset;
+		const varying int32 i2 = End + ParticleOffset;
+
 		const varying FVector P1 = VectorGather(&P[i1]);
 		const varying FVector P2 = VectorGather(&P[i2]);
 
@@ -57,7 +63,187 @@ export void ApplyLongRangeConstraints(uniform FVector P[],
 		const float Length = VectorSize(Direction);
 		Direction = VectorSelect(Length < KINDA_SMALL_NUMBER, ForwardVector, Direction / Length);
 
-		const float Offset = Length - RefLength;
+		const float Offset = Length - RefLength * Scale;
+		const FVector ComputedDelta = Stiffness * Offset * Direction;
+
+		const FVector Delta = VectorSelect(Offset < 0.f, ZeroVector, ComputedDelta);
+
+		VectorScatter(&P[i2], P2 + Delta);
+	}
+}
+
+export void ApplyLongRangeConstraintsStiffnessWeightmap(uniform FVector P[],
+														const uniform FTether Tethers[],
+														const uniform uint8 StiffnessIndices[],
+														const uniform float StiffnessTable[],
+														const uniform float Scale,
+														const uniform int32 NumConstraints,
+														const uniform int32 ParticleOffset)
+{
+	varying FTether Tether;
+
+#if HW_GATHER_SUPPORTED == 1
+	if(programIndex < NumConstraints)
+	{
+		Tether = LoadTetherHelper(&Tethers[0]);
+	}
+#endif
+
+	foreach(i = 0 ... NumConstraints)
+	{
+#if HW_GATHER_SUPPORTED == 0
+		Tether = LoadTetherHelper(&Tethers[extract(i, 0)]);
+#endif
+		const varying int32 Start = Tether.Start;
+		const varying int32 End = Tether.End;
+		const varying float RefLength = Tether.RefLength;
+
+		#pragma ignore warning(perf)
+		const varying uint8 StiffnessIndex = StiffnessIndices[End];
+
+		const varying int32 i1 = Start + ParticleOffset;
+		const varying int32 i2 = End + ParticleOffset;
+
+		const varying FVector P1 = VectorGather(&P[i1]);
+		const varying FVector P2 = VectorGather(&P[i2]);
+
+#if HW_GATHER_SUPPORTED == 1
+		if(i + programCount < NumConstraints)
+		{
+			Tether = LoadTetherHelper(&Tethers[extract(i, 0) + programCount]);
+		}
+#endif
+
+		#pragma ignore warning(perf)
+		const varying float Stiffness = StiffnessTable[StiffnessIndex];
+
+		FVector Direction = P1 - P2;
+		const float Length = VectorSize(Direction);
+		Direction = VectorSelect(Length < KINDA_SMALL_NUMBER, ForwardVector, Direction / Length);
+
+		const float Offset = Length - RefLength * Scale;
+		const FVector ComputedDelta = Stiffness * Offset * Direction;
+
+		const FVector Delta = VectorSelect(Offset < 0.f, ZeroVector, ComputedDelta);
+
+		VectorScatter(&P[i2], P2 + Delta);
+	}
+}
+
+export void ApplyLongRangeConstraintsScaleWeightmap(uniform FVector P[],
+													const uniform FTether Tethers[],
+													const uniform float Stiffness,
+													const uniform uint8 ScaleIndices[],
+													const uniform float ScaleTable[],
+													const uniform int32 NumConstraints,
+													const uniform int32 ParticleOffset)
+{
+	varying FTether Tether;
+
+#if HW_GATHER_SUPPORTED == 1
+	if(programIndex < NumConstraints)
+	{
+		Tether = LoadTetherHelper(&Tethers[0]);
+	}
+#endif
+
+	foreach(i = 0 ... NumConstraints)
+	{
+#if HW_GATHER_SUPPORTED == 0
+		Tether = LoadTetherHelper(&Tethers[extract(i, 0)]);
+#endif
+		const varying int32 Start = Tether.Start;
+		const varying int32 End = Tether.End;
+		const varying float RefLength = Tether.RefLength;
+
+		#pragma ignore warning(perf)
+		const varying uint8 ScaleIndex = ScaleIndices[Tether.End];
+
+		const varying int32 i1 = Start + ParticleOffset;
+		const varying int32 i2 = End + ParticleOffset;
+
+		const varying FVector P1 = VectorGather(&P[i1]);
+		const varying FVector P2 = VectorGather(&P[i2]);
+
+#if HW_GATHER_SUPPORTED == 1
+		if(i + programCount < NumConstraints)
+		{
+			Tether = LoadTetherHelper(&Tethers[extract(i, 0) + programCount]);
+		}
+#endif
+
+		#pragma ignore warning(perf)
+		const varying float Scale = ScaleTable[ScaleIndex];
+
+		FVector Direction = P1 - P2;
+		const float Length = VectorSize(Direction);
+		Direction = VectorSelect(Length < KINDA_SMALL_NUMBER, ForwardVector, Direction / Length);
+
+		const float Offset = Length - RefLength * Scale;
+		const FVector ComputedDelta = Stiffness * Offset * Direction;
+
+		const FVector Delta = VectorSelect(Offset < 0.f, ZeroVector, ComputedDelta);
+
+		VectorScatter(&P[i2], P2 + Delta);
+	}
+}
+
+
+export void ApplyLongRangeConstraintsStiffnessScaleWeightmaps(uniform FVector P[],
+																const uniform FTether Tethers[],
+																const uniform uint8 StiffnessIndices[],
+																const uniform float StiffnessTable[],
+																const uniform uint8 ScaleIndices[],
+																const uniform float ScaleTable[],
+																const uniform int32 NumConstraints,
+																const uniform int32 ParticleOffset)
+{
+	varying FTether Tether;
+
+#if HW_GATHER_SUPPORTED == 1
+	if(programIndex < NumConstraints)
+	{
+		Tether = LoadTetherHelper(&Tethers[0]);
+	}
+#endif
+
+	foreach(i = 0 ... NumConstraints)
+	{
+#if HW_GATHER_SUPPORTED == 0
+		Tether = LoadTetherHelper(&Tethers[extract(i, 0)]);
+#endif
+		const varying int32 Start = Tether.Start;
+		const varying int32 End = Tether.End;
+		const varying float RefLength = Tether.RefLength;
+
+		#pragma ignore warning(perf)
+		const varying uint8 StiffnessIndex = StiffnessIndices[Tether.End];
+		#pragma ignore warning(perf)
+		const varying uint8 ScaleIndex = ScaleIndices[Tether.End];
+
+		const varying int32 i1 = Start + ParticleOffset;
+		const varying int32 i2 = End + ParticleOffset;
+
+		const varying FVector P1 = VectorGather(&P[i1]);
+		const varying FVector P2 = VectorGather(&P[i2]);
+
+#if HW_GATHER_SUPPORTED == 1
+		if(i + programCount < NumConstraints)
+		{
+			Tether = LoadTetherHelper(&Tethers[extract(i, 0) + programCount]);
+		}
+#endif
+
+		#pragma ignore warning(perf)
+		const varying float Stiffness = StiffnessTable[StiffnessIndex];
+		#pragma ignore warning(perf)
+		const varying float Scale = ScaleTable[ScaleIndex];
+
+		FVector Direction = P1 - P2;
+		const float Length = VectorSize(Direction);
+		Direction = VectorSelect(Length < KINDA_SMALL_NUMBER, ForwardVector, Direction / Length);
+
+		const float Offset = Length - RefLength * Scale;
 		const FVector ComputedDelta = Stiffness * Offset * Direction;
 
 		const FVector Delta = VectorSelect(Offset < 0.f, ZeroVector, ComputedDelta);
diff --git a/Engine/Source/Runtime/Experimental/Chaos/Private/Chaos/PBDLongRangeConstraintsBase.cpp b/Engine/Source/Runtime/Experimental/Chaos/Private/Chaos/PBDLongRangeConstraintsBase.cpp
--- a/Engine/Source/Runtime/Experimental/Chaos/Private/Chaos/PBDLongRangeConstraintsBase.cpp
+++ b/Engine/Source/Runtime/Experimental/Chaos/Private/Chaos/PBDLongRangeConstraintsBase.cpp
@@ -5,470 +5,78 @@
 #include "Chaos/Vector.h"
 #include "Chaos/Framework/Parallel.h"
 #include "ChaosStats.h"
+#include "HAL/IConsoleManager.h"
 
 DECLARE_CYCLE_STAT(TEXT("Chaos Cloth Compute Geodesic Constraints"), STAT_ChaosClothComputeGeodesicConstraints, STATGROUP_Chaos);
 
+int32 ChaosPBDLongRangeConstraintsMinParallelBatchSize = 500;
+
+FAutoConsoleVariableRef CVarChaosPBDLongRangeConstraintsMinParallelBatchSize(
+	TEXT("p.Chaos.PBDLongRangeConstraints.MinParallelBatchSize"),
+	ChaosPBDLongRangeConstraintsMinParallelBatchSize,
+	TEXT("The minimum number of long range tethers in a batch to process in parallel."),
+	ECVF_Cheat);
+
 using namespace Chaos;
 
+int32 FPBDLongRangeConstraintsBase::GetMinParallelBatchSize()
+{
+	return ChaosPBDLongRangeConstraintsMinParallelBatchSize;
+}
+
 FPBDLongRangeConstraintsBase::FPBDLongRangeConstraintsBase(
 	const FPBDParticles& Particles,
 	const int32 InParticleOffset,
 	const int32 InParticleCount,
-	const TMap<int32, TSet<int32>>& PointToNeighbors,
+	const TArray<TConstArrayView<TTuple<int32, int32, FRealSingle>>>& InTethers,
 	const TConstArrayView<FRealSingle>& StiffnessMultipliers,
 	const TConstArrayView<FRealSingle>& ScaleMultipliers,
-	const int32 MaxNumTetherIslands,
 	const FVec2& InStiffness,
-	const FVec2& Scale,
-	const EMode InMode)
-	: TethersView(Tethers)
+	const FVec2& InScale)
+	: Tethers(InTethers)
 	, Stiffness(InStiffness, StiffnessMultipliers, InParticleCount)
-	, Mode(InMode)
+	, Scale(InScale)
 	, ParticleOffset(InParticleOffset)
 	, ParticleCount(InParticleCount)
 {
-	switch (Mode)
+	if (ScaleMultipliers.Num() == ParticleCount && ParticleCount > 0)
 	{
-	case EMode::Euclidean:
-		ComputeEuclideanConstraints(Particles, PointToNeighbors, MaxNumTetherIslands);
-		break;
-	case EMode::Geodesic:
-		ComputeGeodesicConstraints(Particles, PointToNeighbors, MaxNumTetherIslands);
-		break;
-	default:
-		unimplemented();
-		break;
-	}
-
-	// Scale the tether's reference lengths
-	const FReal ScaleOffset = FMath::Clamp(Scale[0], 0.01f, 10.f);
+		// Convert the weight maps into an array of lookup indices to the scale table
+		ScaleIndices.AddUninitialized(ParticleCount);
 
-	if (ScaleMultipliers.Num() == InParticleCount)
-	{
-		const FReal ScaleRange = FMath::Clamp(Scale[1], 0.01f, 10.f) - Scale[0];
+		const FRealSingle TableScale = (FRealSingle)(TableSize - 1);
 
-		for (FTether& Tether : Tethers)
+		for (int32 Index = 0; Index < ParticleCount; ++Index)
 		{
-			Tether.RefLength *= ScaleOffset + ScaleRange * FMath::Clamp(ScaleMultipliers[Tether.End], 0.f, 1.f);
+			ScaleIndices[Index] = (uint8)(FMath::Clamp(ScaleMultipliers[Index], (FRealSingle)0., (FRealSingle)1.) * TableScale);
 		}
-	}
-	else
-	{
-		for (FTether& Tether : Tethers)
-		{
-			Tether.RefLength *= ScaleOffset;
-		}
-	}
-}
-
-TArray<TArray<int32>> FPBDLongRangeConstraintsBase::ComputeIslands(
-    const TMap<int32, TSet<int32>>& PointToNeighbors,
-    const TArray<int32>& KinematicParticles)
-{
-	// Compute Islands
-	int32 NextIsland = 0;
-	TArray<int32> FreeIslands;
-	TArray<TArray<int32>> IslandElements;
 
-	TMap<int32, int32> ParticleToIslandMap;
-	ParticleToIslandMap.Reserve(KinematicParticles.Num());
-
-	for (const int32 Element : KinematicParticles)
-	{
-		// Assign Element an island, possibly unionizing existing islands
-		int32 Island = TNumericLimits<int32>::Max();
-
-		// KinematicParticles is generated from the keys of PointToNeighbors, so
-		// we don't need to check if Element exists.
-		const TSet<int32>& Neighbors = PointToNeighbors[Element];
-		for (const int32 Neighbor : Neighbors)
-		{
-			if (ParticleToIslandMap.Contains(Neighbor))
-			{
-				if (Island == TNumericLimits<int32>::Max())
-				{
-					// We don't have an assigned island yet.  Join with the island
-					// of the neighbor.
-					Island = ParticleToIslandMap[Neighbor];
-				}
-				else
-				{
-					const int32 OtherIsland = ParticleToIslandMap[Neighbor];
-					if (OtherIsland != Island)
-					{
-						// This kinematic particle is connected to multiple islands.
-						// Union them.
-						for (const int32 OtherElement : IslandElements[OtherIsland])
-						{
-							check(ParticleToIslandMap[OtherElement] == OtherIsland);
-							ParticleToIslandMap[OtherElement] = Island;
-						}
-						IslandElements[Island].Append(IslandElements[OtherIsland]);
-						IslandElements[OtherIsland].Reset(); // Don't deallocate
-						FreeIslands.AddUnique(OtherIsland);
-					}
-				}
-			}
-		}
-
-		// If no connected Island was found, create a new one (or reuse an old 
-		// vacated one).
-		if (Island == TNumericLimits<int32>::Max())
-		{
-			if (FreeIslands.Num() == 0)
-			{
-				Island = NextIsland++;
-				IslandElements.SetNum(NextIsland);
-			}
-			else
-			{
-				// Reuse a previously allocated, but currently unused, island.
-				Island = FreeIslands.Pop();
-			}
-		}
-
-		ParticleToIslandMap.FindOrAdd(Element) = Island;
-		check(IslandElements.IsValidIndex(Island));
-		IslandElements[Island].Add(Element);
+		// Initialize empty table until ApplyValues is called
+		ScaleTable.AddZeroed(TableSize);
 	}
-
-	// IslandElements may contain empty arrays.
-	for (int32 IslandIndex = 0; IslandIndex < IslandElements.Num(); )
+	else
 	{
-		if (!IslandElements[IslandIndex].Num())
-		{
-			IslandElements.RemoveAtSwap(IslandIndex, 1, false);  // RemoveAtSwap takes the last elements to replace the current one, do not increment the index in this case
-		}
-		else
-		{
-			++IslandIndex;
-		}
+		// Initialize with a one element table until ApplyValues is called
+		ScaleIndices.AddZeroed(1);
+		ScaleTable.AddZeroed(1);
 	}
-
-	return IslandElements;
 }
 
-void FPBDLongRangeConstraintsBase::ComputeEuclideanConstraints(
-    const FPBDParticles& Particles,
-    const TMap<int32, TSet<int32>>& PointToNeighbors,
-    const int32 MaxNumTetherIslands)
+void FPBDLongRangeConstraintsBase::ApplyProperties(const FReal Dt, const int32 NumIterations)
 {
-	// Fill up the list of all used indices
-	TArray<int32> Nodes;
-	PointToNeighbors.GenerateKeyArray(Nodes);
-
-	TArray<int32> KinematicParticles;
-	for (const int32 Node : Nodes)
-	{
-		if (Particles.InvM(Node) == (FReal)0.)
-		{
-			KinematicParticles.Add(Node);
-		}
-	}
-
-	// Compute the islands of kinematic particles
-	const TArray<TArray<int32>> IslandElements = ComputeIslands(PointToNeighbors, KinematicParticles);
-	int32 NumTotalIslandElements = 0;
-	for (const TArray<int32>& Elements : IslandElements)
-	{
-		NumTotalIslandElements += Elements.Num();
-	}
-
-	TArray<TArray<FTether>> NewTethersSlots;
-	NewTethersSlots.SetNum(Nodes.Num());
-
-	for (TArray<FTether>& NewTethers : NewTethersSlots)
-	{
-		NewTethers.Reserve(MaxNumTetherIslands);
-	}
-
-	int32 NumTethers = 0;
-
-	PhysicsParallelFor(Nodes.Num(), [&](int32 Index)
-	{
-		const int32 Node = Nodes[Index];
-
-		// For each non-kinematic particle i...
-		if (Particles.InvM(Node) == (FReal)0.)
-		{
-			return;
-		}
-
-		// Measure the distance to all kinematic particles in all islands...
-		TArray<Pair<FReal, int32>> ClosestElements;
-		ClosestElements.Reserve(NumTotalIslandElements);
-
-		for (const TArray<int32>& Elements : IslandElements)
-		{
-			// IslandElements may contain empty arrays.
-			if (Elements.Num() == 0)
-			{
-				continue;
-			}
-
-			int32 ClosestElement = TNumericLimits<int32>::Max();
-			FReal ClosestSquareDistance = TNumericLimits<FReal>::Max();
-			for (const int32 Element : Elements)
-			{
-				const FReal SquareDistance = (Particles.X(Element) - Particles.X(Node)).SizeSquared();
-				if (SquareDistance < ClosestSquareDistance)
-				{
-					ClosestElement = Element;
-					ClosestSquareDistance = SquareDistance;
-				}
-			}
-			ClosestElements.Add(MakePair(FMath::Sqrt(ClosestSquareDistance), ClosestElement));
-		}
-
-		// Order all by distance, smallest first...
-		ClosestElements.Sort();
-
-		// Take the first MaxNumTetherIslands.
-		if (MaxNumTetherIslands < ClosestElements.Num())
-		{
-			ClosestElements.SetNum(MaxNumTetherIslands);
-		}
-
-		// Add constraints between this particle, and the N closest...
-		for (const Pair<FReal, int32> Element : ClosestElements)
-		{
-			NewTethersSlots[Index].Emplace(Element.Second, Node, Element.First);
-			++NumTethers;
-		}
-	});
-
-	// Consolidate the N slots into one single list of tethers, spreading the nodes across different ranges for parallel processing
-	Tethers.Reset(NumTethers);
-
-	for (int32 Slot = 0; Slot < MaxNumTetherIslands; ++Slot)
-	{
-		int32 Size = 0;
-		for (int32 Index = 0; Index < Nodes.Num(); ++Index)
-		{
-			if (NewTethersSlots[Index].IsValidIndex(Slot))
-			{
-				Tethers.Emplace(NewTethersSlots[Index][Slot]);
-				++Size;
-			}
-		}
-		TethersView.AddRange(Size);
+	// Apply stiffness
+	Stiffness.ApplyValues(Dt, NumIterations);
+
+	// Apply scale
+	const FReal Offset = Scale[0];
+	const FReal Range = Scale[1] - Scale[0];
+	const int32 ScaleTableSize = ScaleTable.Num();
+	const FReal WeightIncrement = (ScaleTableSize > 1) ? (FReal)1. / (FReal)(ScaleTableSize - 1) : (FReal)1.; // Must allow full range from 0 to 1 included
+	for (int32 Index = 0; Index < ScaleTableSize; ++Index)
+	{
+		const FReal Weight = (FReal)Index * WeightIncrement;
+		ScaleTable[Index] = Offset + Weight * Range;
 	}
 }
 
-void FPBDLongRangeConstraintsBase::ComputeGeodesicConstraints(
-    const FPBDParticles& Particles,
-    const TMap<int32, TSet<int32>>& PointToNeighbors,
-    const int32 MaxNumTetherIslands)
-{
-	TRACE_CPUPROFILER_EVENT_SCOPE(FPBDLongRangeConstraintsBase_ComputeGeodesicConstraints);
-	SCOPE_CYCLE_COUNTER(STAT_ChaosClothComputeGeodesicConstraints);
-
-	// Fill up the list of all used indices
-	TArray<int32> Nodes;
-	PointToNeighbors.GenerateKeyArray(Nodes);
-
-	// Find all kinematic points to use as anchor points for the tethers and seed the path finding
-	TArray<int32> Seeds;  // Kinematic nodes connected to at least one dynamic node
-	Seeds.Reserve(Nodes.Num() / 10);  // Start at 10% seeds to minimize this array's reallocation
-
-	int32 NumDynamicNodes = 0;
-
-	for (const int32 Node : Nodes)
-	{
-		if (Particles.InvM(Node) == (FReal)0.)
-		{
-			const TSet<int32>& Neighbors = PointToNeighbors[Node];
-			for (const int32 Neighbor : Neighbors)
-			{
-				if (Particles.InvM(Neighbor) != (FReal)0.)  // There's no point to keep a node as an anchor, unless at least one of its neighbors is dynamic
-				{
-					Seeds.Add(Node);
-					break;
-				}
-			}
-		}
-		else
-		{
-			++NumDynamicNodes;
-		}
-	}
-
-	// Dijkstra for each Kinematic Particle (assume a small number of kinematic points) - note this is N^2 log N with N kinematic points
-	// Find shortest paths from each kinematic node to all dynamic nodes
-	TMap<TVec2<int32>, FReal> GeodesicDistances;
-	GeodesicDistances.Reserve(Seeds.Num() * NumDynamicNodes);
-
-	for (const int32 Seed : Seeds)
-	{
-		for (const int32 Node : Nodes)
-		{
-			if (Particles.InvM(Node) != (FReal)0.)
-			{
-				GeodesicDistances.Emplace(TVec2<int32>(Seed, Node), Node != Seed ? TNumericLimits<FReal>::Max() : (FReal)0.);
-			}
-		}
-	}
-
-	PhysicsParallelFor(Seeds.Num(), [&](int32 Index)
-	{
-		const int32 Seed = Seeds[Index];
-
-		// Keep track of all visited nodes in a bit array, will need to be offsetted since the first node index may not be 0
-		TBitArray<> VisitedNodes(false, Nodes.Num());
-
-		// Priority queue based implementation of the Dijkstra algorithm
-#define CHAOS_PBDLONGRANGEATTACHMENT_PROFILE_HEAPSIZE 0
-#if CHAOS_PBDLONGRANGEATTACHMENT_PROFILE_HEAPSIZE
-		static int32 MaxHeapSize = 1;  // Record the max heap size
-#else
-		static const int32 MaxHeapSize = 512;  // Set the queue size to something large enough to avoid reallocations in most cases
-#endif
-		auto LessPredicate = [Seed, &GeodesicDistances](int32 Node1, int32 Node2) -> bool
-			{
-				return GeodesicDistances[TVec2<int32>(Seed, Node1)] < GeodesicDistances[TVec2<int32>(Seed, Node2)];  // Less for node priority
-			};
-
-		TArray<int32> Queue;
-		Queue.Reserve(MaxHeapSize);
-		Queue.Heapify(LessPredicate);  // Turn the array into a priority queue
-
-		// Initiate the graph progression
-		VisitedNodes[Seed - ParticleOffset] = true;
-		Queue.HeapPush(Seed, LessPredicate);
-
-		do
-		{
-			int32 ParentNode;
-			Queue.HeapPop(ParentNode, LessPredicate, false);
-
-			check(VisitedNodes[ParentNode - ParticleOffset]);
-
-			const FReal ParentDistance = (ParentNode != Seed) ? GeodesicDistances[TVec2<int32>(Seed, ParentNode)] : (FReal)0.;
-
-			const TSet<int32>& NeighborNodes = PointToNeighbors[ParentNode];
-			for (const int32 NeighborNode : NeighborNodes)
-			{
-				check(NeighborNode != ParentNode);
-
-				// Do not progress onto kinematic nodes
-				if (Particles.InvM(NeighborNode) == (FReal)0.)
-				{
-					continue;
-				}
-
-				// Update the geodesic distance if this path is a shorter one
-				const FReal NewDistance = ParentDistance + (Particles.X(NeighborNode) - Particles.X(ParentNode)).Size();
-
-				FReal& GeodesicDistance = GeodesicDistances[TVec2<int32>(Seed, NeighborNode)];
 
-				if (NewDistance < GeodesicDistance)
-				{
-					// Update this path distance
-					GeodesicDistance = NewDistance;
-
-					// Progress to this node position if it hasn't yet been visited
-					if (!VisitedNodes[NeighborNode - ParticleOffset])
-					{
-						VisitedNodes[NeighborNode - ParticleOffset] = true;
-
-						Queue.HeapPush(NeighborNode, LessPredicate);
-
-#if CHAOS_PBDLONGRANGEATTACHMENT_PROFILE_HEAPSIZE
-						MaxHeapSize = FMath::Max(Queue.Num(), MaxHeapSize);
-#endif
-					}
-				}
-			}
-		} while (Queue.Num());
-	});
-
-	// Compute islands of kinematic points
-	const TArray<TArray<int32>> Islands = ComputeIslands(PointToNeighbors, Seeds);
-
-	// Find the tether constraints starting from each dynamic node
-	TArray<TArray<FTether>> NewTethersSlots;
-	NewTethersSlots.SetNum(Nodes.Num());
-
-	for (TArray<FTether>& NewTethers : NewTethersSlots)
-	{
-		NewTethers.Reserve(MaxNumTetherIslands);
-	}
-
-	TAtomic<int32> NumTethers(0);
-
-	PhysicsParallelFor(Nodes.Num(), [&](int32 Index)
-	{
-		const int32 Node = Nodes[Index];
-		if (Particles.InvM(Node) == (FReal)0.)
-		{
-			return;  // No tethers on dynamic particles
-		}
-
-		// Find the closest seeds in each island
-		TArray<int32, TInlineAllocator<32>> ClosestSeeds;
-		ClosestSeeds.Reserve(Islands.Num());
-
-		for (const TArray<int32>& Seeds : Islands)
-		{
-			if (!Seeds.Num())
-			{
-				continue;  // Empty island 
-			}
-
-			int32 ClosestSeed = INDEX_NONE;
-			FReal ClosestDistance = TNumericLimits<FReal>::Max();
-
-			for (const int32 Seed : Seeds)
-			{
-				const FReal Distance = GeodesicDistances[TVec2<int32>(Seed, Node)];
-				if (Distance < ClosestDistance)
-				{
-					ClosestDistance = Distance;
-					ClosestSeed = Seed;
-				}
-			}
-			if (ClosestSeed != INDEX_NONE)
-			{
-				ClosestSeeds.Add(ClosestSeed);
-			}
-		}
-
-		// Sort all the tethers for this node based on smallest distance
-		ClosestSeeds.Sort([&GeodesicDistances, Node](int32 Seed1, int32 Seed2)
-			{
-				return GeodesicDistances[TVec2<int32>(Seed1, Node)] < GeodesicDistances[TVec2<int32>(Seed2, Node)];
-			});
-
-		// Keep only N closest tethers
-		if (MaxNumTetherIslands < ClosestSeeds.Num())
-		{
-			ClosestSeeds.SetNum(MaxNumTetherIslands);
-		}
-
-		// Add these tethers to the N (or less) available slots
-		for (const int32 Seed : ClosestSeeds)
-		{
-			const FReal RefLength = GeodesicDistances[TVec2<int32>(Seed, Node)];
-			NewTethersSlots[Index].Emplace(Seed, Node, RefLength);
-			++NumTethers;
-		}
-		check(NewTethersSlots[Index].Num() <= MaxNumTetherIslands);
-	});
-
-	// Consolidate the N slots into one single list of tethers, spreading the nodes across different ranges for parallel processing
-	Tethers.Reset(NumTethers);
-
-	for (int32 Slot = 0; Slot < MaxNumTetherIslands; ++Slot)
-	{
-		int32 Size = 0;
-		for (int32 Index = 0; Index < Nodes.Num(); ++Index)
-		{
-			if (NewTethersSlots[Index].IsValidIndex(Slot))
-			{
-				Tethers.Emplace(NewTethersSlots[Index][Slot]);
-				++Size;
-			}
-		}
-		TethersView.AddRange(Size);
-	}
-}
diff --git a/Engine/Source/Runtime/Experimental/Chaos/Public/Chaos/PBDLongRangeConstraints.h b/Engine/Source/Runtime/Experimental/Chaos/Public/Chaos/PBDLongRangeConstraints.h
--- a/Engine/Source/Runtime/Experimental/Chaos/Public/Chaos/PBDLongRangeConstraints.h
+++ b/Engine/Source/Runtime/Experimental/Chaos/Public/Chaos/PBDLongRangeConstraints.h
@@ -13,30 +13,27 @@ class CHAOS_API FPBDLongRangeConstraints final : public FPBDLongRangeConstraints
 public:
 	typedef FPBDLongRangeConstraintsBase Base;
 	typedef typename Base::FTether FTether;
-	typedef typename Base::EMode EMode;
 
 	FPBDLongRangeConstraints(
 		const FPBDParticles& Particles,
 		const int32 InParticleOffset,
 		const int32 InParticleCount,
-		const TMap<int32, TSet<int32>>& PointToNeighbors,
+		const TArray<TConstArrayView<TTuple<int32, int32, FRealSingle>>>& InTethers,
 		const TConstArrayView<FRealSingle>& StiffnessMultipliers,
 		const TConstArrayView<FRealSingle>& ScaleMultipliers,
-		const int32 MaxNumTetherIslands = 4,
 		const FVec2& InStiffness = FVec2::UnitVector,
-		const FVec2& Scale = FVec2::UnitVector,
-		const EMode InMode = EMode::Geodesic)
-		: FPBDLongRangeConstraintsBase(Particles, InParticleOffset, InParticleCount, PointToNeighbors, StiffnessMultipliers, ScaleMultipliers, MaxNumTetherIslands, InStiffness, Scale, InMode) {}
+		const FVec2& InScale = FVec2::UnitVector)
+		: FPBDLongRangeConstraintsBase(Particles, InParticleOffset, InParticleCount, InTethers, StiffnessMultipliers, ScaleMultipliers, InStiffness, InScale) {}
 	virtual ~FPBDLongRangeConstraints() {}
 
-	void Apply(FPBDParticles& Particles, const FReal Dt, const TArray<int32>& ConstraintIndices) const;
 	void Apply(FPBDParticles& Particles, const FReal Dt) const;
 
 private:
 	using Base::Tethers;
-	using Base::TethersView;
 	using Base::Stiffness;
 	using Base::ParticleOffset;
+	using Base::ScaleTable;
+	using Base::ScaleIndices;
 };
 }
 
diff --git a/Engine/Source/Runtime/Experimental/Chaos/Public/Chaos/PBDLongRangeConstraintsBase.h b/Engine/Source/Runtime/Experimental/Chaos/Public/Chaos/PBDLongRangeConstraintsBase.h
--- a/Engine/Source/Runtime/Experimental/Chaos/Public/Chaos/PBDLongRangeConstraintsBase.h
+++ b/Engine/Source/Runtime/Experimental/Chaos/Public/Chaos/PBDLongRangeConstraintsBase.h
@@ -24,73 +24,98 @@ public:
 		AccurateTetherAccurateLength = Geodesic
 	};
 
-	struct FTether
-	{
-		int32 Start;
-		int32 End;
-		FReal RefLength;
-
-		FTether(int32 InStart, int32 InEnd, FReal InRefLength): Start(InStart), End(InEnd), RefLength(InRefLength) {}
-
-		inline FVec3 GetDelta(const FPBDParticles& Particles) const
-		{
-			checkSlow(Particles.InvM(Start) == (FReal)0.);
-			checkSlow(Particles.InvM(End) > (FReal)0.);
-			FVec3 Direction = Particles.P(Start) - Particles.P(End);
-			const FReal Length = Direction.SafeNormalize();
-			const FReal Offset = Length - RefLength;
-			return Offset < (FReal)0. ? FVec3((FReal)0.) : Offset * Direction;
-		};
-
-		inline void GetDelta(const FPBDParticles& Particles, FVec3& OutDirection, FReal& OutOffset) const
-		{
-			checkSlow(Particles.InvM(Start) == (FReal)0.);
-			checkSlow(Particles.InvM(End) > (FReal)0.);
-			OutDirection = Particles.P(Start) - Particles.P(End);
-			const FReal Length = OutDirection.SafeNormalize();
-			OutOffset = FMath::Max((FReal)0., Length - RefLength);
-		};
-	};
+	typedef TTuple<int32, int32, float> FTether;
 
 	FPBDLongRangeConstraintsBase(
 		const FPBDParticles& Particles,
 		const int32 InParticleOffset,
 		const int32 InParticleCount,
-		const TMap<int32, TSet<int32>>& PointToNeighbors,
+		const TArray<TConstArrayView<TTuple<int32, int32, FRealSingle>>>& InTethers,
 		const TConstArrayView<FRealSingle>& StiffnessMultipliers,
 		const TConstArrayView<FRealSingle>& ScaleMultipliers,
-		const int32 MaxNumTetherIslands = 4,
 		const FVec2& InStiffness = FVec2::UnitVector,
-		const FVec2& Scale = FVec2::UnitVector,
-		const EMode InMode = EMode::AccurateTetherFastLength);
+		const FVec2& InScale = FVec2::UnitVector);
 
 	virtual ~FPBDLongRangeConstraintsBase() {}
 
-	EMode GetMode() const { return Mode; }
-
 	// Return the stiffness input values used by the constraint
 	FVec2 GetStiffness() const { return Stiffness.GetWeightedValue(); }
 
 	// Set the stiffness input values used by the constraint
 	void SetStiffness(const FVec2& InStiffness) { Stiffness.SetWeightedValue(InStiffness); }
 
+	// Set the scale low and high value of the scale weight map
+	void SetScale(const FVec2& InScale) { Scale = InScale.ClampAxes((FReal)0.01, (FReal)10.); }
+
 	// Set stiffness offset and range, as well as the simulation stiffness exponent
-	void ApplyProperties(const FReal Dt, const int32 NumIterations) { Stiffness.ApplyValues(Dt, NumIterations); }
+	void ApplyProperties(const FReal Dt, const int32 NumIterations);
+
+	// Return the tethers, organized in concurent friendly batches
+	const TArray<TConstArrayView<FTether>>& GetTethers() const { return Tethers; }
+
+	// Return the start index of the specified tether
+	int32 GetStartIndex(const FTether& Tether) const { return Tether.Get<0>(); }
+
+	// Return the kinematic particle index of the specified tether
+	int32 GetStartParticle(const FTether& Tether) const { return GetStartIndex(Tether) + ParticleOffset; }
+
+	// Return the end index of the specified tether
+	int32 GetEndIndex(const FTether& Tether) const { return Tether.Get<1>(); }
+
+	// Return the dynamic particle index of the specified tether
+	int32 GetEndParticle(const FTether& Tether) const { return GetEndIndex(Tether) + ParticleOffset; }
 
-	const TArray<FTether>& GetTethers() const { return Tethers; }
+	// Return the reference length of the specified tether
+	FReal GetRefLength(const FTether& Tether) const { return (FReal)Tether.Get<2>(); }
 
-	static TArray<TArray<int32>> ComputeIslands(const TMap<int32, TSet<int32>>& PointToNeighbors, const TArray<int32>& KinematicParticles);
+	// Return the Tether scale for the specified tether
+	FReal GetScale(const FTether& Tether) const { return HasScaleWeightMap() ? ScaleTable[ScaleIndices[GetEndIndex(Tether)]] : ScaleTable[0]; }
+
+	// Return the target length of the specified tether (= RefLength * Scale)
+	FReal GetTargetLength(const FTether& Tether) const { return GetRefLength(Tether) * GetScale(Tether); }
 
 protected:
-	void ComputeEuclideanConstraints(const FPBDParticles& Particles, const TMap<int32, TSet<int32>>& PointToNeighbors, const int32 NumberOfAttachments);
-	void ComputeGeodesicConstraints(const FPBDParticles& Particles, const TMap<int32, TSet<int32>>& PointToNeighbors, const int32 NumberOfAttachments);
+	// Return the minimum number of long range tethers in a batch to process in parallel
+	static int32 GetMinParallelBatchSize();
+
+	// Return whether the constraint has been setup with a weightmap to interpolate between two low and high values of scales
+	bool HasScaleWeightMap() const { return ScaleTable.Num() > 1; }
+
+	// Return a vector representing the amount of segment required for the tether to shrink back to its maximum target length constraint, or zero if the constraint is already met
+	inline FVec3 GetDelta(const FPBDParticles& Particles, const FTether& Tether, const FReal InScale) const
+	{
+		const int32 Start = GetStartParticle(Tether);
+		const int32 End = GetEndParticle(Tether);
+		const FReal TargetLength = GetRefLength(Tether) * InScale;
+		checkSlow(Particles.InvM(Start) == (FReal)0.);
+		checkSlow(Particles.InvM(End) > (FReal)0.);
+		FVec3 Direction = Particles.P(Start) - Particles.P(End);
+		const FReal Length = Direction.SafeNormalize();
+		const FReal Offset = Length - TargetLength;
+		return Offset < (FReal)0. ? FVec3((FReal)0.) : Offset * Direction;
+	};
+
+	// Return a direction and length representing the amount of segment required for the tether to shrink back to its maximum target length constraint, or zero if the constraint is already met
+	inline void GetDelta(const FPBDParticles& Particles, const FTether& Tether, const FReal InScale, FVec3& OutDirection, FReal& OutOffset) const
+	{
+		const int32 Start = GetStartParticle(Tether);
+		const int32 End = GetEndParticle(Tether);
+		const FReal TargetLength = GetRefLength(Tether) * InScale;
+		checkSlow(Particles.InvM(Start) == (FReal)0.);
+		checkSlow(Particles.InvM(End) > (FReal)0.);
+		OutDirection = Particles.P(Start) - Particles.P(End);
+		const FReal Length = OutDirection.SafeNormalize();
+		OutOffset = FMath::Max((FReal)0., Length - TargetLength);
+	};
 
 protected:
-	TArray<FTether> Tethers;
-	TPBDActiveView<TArray<FTether>> TethersView;
-	FPBDStiffness Stiffness;
-	const EMode Mode;
-	const int32 ParticleOffset;
-	const int32 ParticleCount;
+	static constexpr int32 TableSize = 16;  // The size of the weightmaps lookup table
+	const TArray<TConstArrayView<FTether>>& Tethers;  // Array view on the tether provided to this constraint
+	FPBDStiffness Stiffness;  // Stiffness weightmap lookup table 
+	TArray<uint8> ScaleIndices;  // Per particle array of index to the scale lookup table
+	TArray<FReal> ScaleTable;  // Fixed lookup table of scale values
+	FVec2 Scale;  // High and low values of tether target length scale used by the the lookup tables interpolation
+	const int32 ParticleOffset;  // Index of the first usable particle
+	const int32 ParticleCount;  // Number of particles available to this constraint
 };
 }
diff --git a/Engine/Source/Runtime/Experimental/Chaos/Public/Chaos/XPBDLongRangeConstraints.h b/Engine/Source/Runtime/Experimental/Chaos/Public/Chaos/XPBDLongRangeConstraints.h
--- a/Engine/Source/Runtime/Experimental/Chaos/Public/Chaos/XPBDLongRangeConstraints.h
+++ b/Engine/Source/Runtime/Experimental/Chaos/Public/Chaos/XPBDLongRangeConstraints.h
@@ -17,22 +17,24 @@ class FXPBDLongRangeConstraints final : public FPBDLongRangeConstraintsBase
 public:
 	typedef FPBDLongRangeConstraintsBase Base;
 	typedef typename Base::FTether FTether;
-	typedef typename Base::EMode EMode;
 
 	FXPBDLongRangeConstraints(
 		const FPBDParticles& Particles,
 		const int32 InParticleOffset,
 		const int32 InParticleCount,
-		const TMap<int32, TSet<int32>>& PointToNeighbors,
+		const TArray<TConstArrayView<TTuple<int32, int32, FRealSingle>>>& InTethers,
 		const TConstArrayView<FRealSingle>& StiffnessMultipliers,
 		const TConstArrayView<FRealSingle>& ScaleMultipliers,
-		const int32 MaxNumTetherIslands = 4,
 		const FVec2& InStiffness = FVec2::UnitVector,
-		const FVec2& Scale = FVec2::UnitVector,
-		const EMode InMode = EMode::Geodesic)
-	    : FPBDLongRangeConstraintsBase(Particles, InParticleOffset, InParticleCount, PointToNeighbors, StiffnessMultipliers, ScaleMultipliers, MaxNumTetherIslands, InStiffness, Scale, InMode)
+		const FVec2& InScale = FVec2::UnitVector)
+	    : FPBDLongRangeConstraintsBase(Particles, InParticleOffset, InParticleCount, InTethers, StiffnessMultipliers, ScaleMultipliers, InStiffness, InScale)
 	{
-		Lambdas.Reserve(Tethers.Num());
+		int32 NumTethers = 0;
+		for (const TConstArrayView<FTether>& TetherBatch : Tethers)
+		{
+			NumTethers += TetherBatch.Num();
+		}
+		Lambdas.Reserve(NumTethers);
 	}
 
 	virtual ~FXPBDLongRangeConstraints() {}
@@ -47,68 +49,100 @@ public:
 	{
 		SCOPE_CYCLE_COUNTER(STAT_XPBD_LongRange);
 		// Run particles in parallel, and ranges in sequence to avoid a race condition when updating the same particle from different tethers
-		static const int32 MinParallelSize = 500;
+		const int32 MinParallelSize = GetMinParallelBatchSize();
+
 		if (Stiffness.HasWeightMap())
 		{
-			TethersView.ParallelFor([this, &Particles, Dt](TArray<FTether>& /*InTethers*/, int32 Index)
+			if (HasScaleWeightMap())
+			{
+				int32 ConstraintOffset = 0;
+				for (const TConstArrayView<FTether>& TetherBatch : Tethers)
 				{
-					const FReal ExpStiffnessValue = Stiffness[Index - ParticleOffset];
-					Apply(Particles, Dt, Index, ExpStiffnessValue);
-				}, MinParallelSize);
+					PhysicsParallelFor(TetherBatch.Num(), [this, &Particles, Dt, &TetherBatch, ConstraintOffset](int32 Index)
+						{
+							const FTether& Tether = TetherBatch[Index];
+							const int32 LocalParticleIndex = GetEndIndex(Tether);
+							const FReal Scale = ScaleTable[ScaleIndices[LocalParticleIndex]];
+							const FReal ExpStiffnessValue = Stiffness[LocalParticleIndex];
+							Apply(Particles, Dt, Tether, ConstraintOffset + Index, ExpStiffnessValue, Scale);
+						}, TetherBatch.Num() < MinParallelSize);
+					ConstraintOffset += TetherBatch.Num();
+				}
+			}
+			else
+			{
+				const FReal ScaleValue = ScaleTable[0];
+				int32 ConstraintOffset = 0;
+				for (const TConstArrayView<FTether>& TetherBatch : Tethers)
+				{
+					PhysicsParallelFor(TetherBatch.Num(), [this, &Particles, Dt, &TetherBatch, ScaleValue, ConstraintOffset](int32 Index)
+						{
+							const FTether& Tether = TetherBatch[Index];
+							const int32 LocalParticleIndex = GetEndIndex(Tether);
+							const FReal ExpStiffnessValue = Stiffness[LocalParticleIndex];
+							Apply(Particles, Dt, Tether, ConstraintOffset + Index, ExpStiffnessValue, ScaleValue);
+						}, TetherBatch.Num() < MinParallelSize);
+					ConstraintOffset += TetherBatch.Num();
+				}
+			}
 		}
 		else
 		{
 			const FReal ExpStiffnessValue = (FReal)Stiffness;
-			TethersView.ParallelFor([this, &Particles, Dt, ExpStiffnessValue](TArray<FTether>& /*InTethers*/, int32 Index)
-				{
-					Apply(Particles, Dt, Index, ExpStiffnessValue);
-				}, MinParallelSize);
-		}
-	}
 
-	void Apply(FPBDParticles& Particles, const FReal Dt, const TArray<int32>& ConstraintIndices) const
-	{
-		SCOPE_CYCLE_COUNTER(STAT_XPBD_LongRange);
-		if (Stiffness.HasWeightMap())
-		{
-			for (const int32 Index : ConstraintIndices)
+			if (HasScaleWeightMap())
 			{
-				const FReal ExpStiffnessValue = Stiffness[Index - ParticleOffset];
-				Apply(Particles, Dt, Index, ExpStiffnessValue);
+				int32 ConstraintOffset = 0;
+				for (const TConstArrayView<FTether>& TetherBatch : Tethers)
+				{
+					PhysicsParallelFor(TetherBatch.Num(), [this, &Particles, Dt, &TetherBatch, ExpStiffnessValue, ConstraintOffset](int32 Index)
+						{
+							const FTether& Tether = TetherBatch[Index];
+							const int32 LocalParticleIndex = GetEndIndex(Tether);
+							const FReal Scale = ScaleTable[ScaleIndices[LocalParticleIndex]];
+							Apply(Particles, Dt, Tether, ConstraintOffset + Index, ExpStiffnessValue, Scale);
+						}, TetherBatch.Num() < MinParallelSize);
+					ConstraintOffset += TetherBatch.Num();
+				}
 			}
-		}
-		else
-		{
-			const FReal ExpStiffnessValue = (FReal)Stiffness;
-			for (const int32 Index : ConstraintIndices)
+			else
 			{
-				Apply(Particles, Dt, Index, ExpStiffnessValue);
+				const FReal ScaleValue = ScaleTable[0];
+				int32 ConstraintOffset = 0;
+				for (const TConstArrayView<FTether>& TetherBatch : Tethers)
+				{
+					PhysicsParallelFor(TetherBatch.Num(), [this, &Particles, Dt, &TetherBatch, ExpStiffnessValue, ScaleValue, ConstraintOffset](int32 Index)
+						{
+							const FTether& Tether = TetherBatch[Index];
+							Apply(Particles, Dt, Tether, ConstraintOffset + Index, ExpStiffnessValue, ScaleValue);
+						}, TetherBatch.Num() < MinParallelSize);
+					ConstraintOffset += TetherBatch.Num();
+				}
 			}
 		}
 	}
 
 private:
-	void Apply(FPBDParticles& Particles, const FReal Dt, int32 Index, const FReal InStiffness) const
+	void Apply(FPBDParticles& Particles, const FReal Dt, const FTether& Tether, int32 ConstraintIndex, const FReal InStiffness, const FReal InScale) const
 	{
-		const FTether& Tether = Tethers[Index];
-
 		FVec3 Direction;
 		FReal Offset;
-		Tether.GetDelta(Particles, Direction, Offset);
+		GetDelta(Particles, Tether, InScale, Direction, Offset);
 
-		FReal& Lambda = Lambdas[Index];
+		FReal& Lambda = Lambdas[ConstraintIndex];
 		const FReal Alpha = (FReal)XPBDLongRangeMaxCompliance / (InStiffness * Dt * Dt);
 
 		const FReal DLambda = (Offset - Alpha * Lambda) / ((FReal)1. + Alpha);
-		Particles.P(Tether.End) += DLambda * Direction;
+		Particles.P(GetEndParticle(Tether)) += DLambda * Direction;
 		Lambda += DLambda;
 	}
 
 private:
 	using Base::Tethers;
-	using Base::TethersView;
 	using Base::Stiffness;
 	using Base::ParticleOffset;
+	using Base::ScaleTable;
+	using Base::ScaleIndices;
 
 	mutable TArray<FReal> Lambdas;
 };
