Cloth - Fix erroneous Max operation between bool and float on self-collision radius.

diff --git a/Engine/Source/Runtime/ClothingSystemRuntimeCommon/Private/ClothPhysicalMeshData.cpp b/Engine/Source/Runtime/ClothingSystemRuntimeCommon/Private/ClothPhysicalMeshData.cpp
--- a/Engine/Source/Runtime/ClothingSystemRuntimeCommon/Private/ClothPhysicalMeshData.cpp
+++ b/Engine/Source/Runtime/ClothingSystemRuntimeCommon/Private/ClothPhysicalMeshData.cpp
@@ -91,7 +91,7 @@ void FClothPhysicalMeshData::BuildSelfCollisionData(const TMap<FName, TObjectPtr
 	float SelfCollisionRadius = 0.f;
 	for (const TPair<FName, TObjectPtr<UClothConfigBase>>& ClothConfig : ClothConfigs)
 	{
-		SelfCollisionRadius = FMath::Max(SelfCollisionRadius, ClothConfig.Value->NeedsSelfCollisionData());
+		SelfCollisionRadius = FMath::Max(SelfCollisionRadius, ClothConfig.Value->GetSelfCollisionRadius());
 	}
 	if (SelfCollisionRadius)
 	{
