Chaos Cloth - Fix/skip tether generation when there isn't any kinematic nodes.

diff --git a/Engine/Source/Runtime/ClothingSystemRuntimeCommon/Private/ClothTetherData.cpp b/Engine/Source/Runtime/ClothingSystemRuntimeCommon/Private/ClothTetherData.cpp
--- a/Engine/Source/Runtime/ClothingSystemRuntimeCommon/Private/ClothTetherData.cpp
+++ b/Engine/Source/Runtime/ClothingSystemRuntimeCommon/Private/ClothTetherData.cpp
@@ -134,39 +134,42 @@ FClothTetherDataPrivate::FClothTetherDataPrivate(
 		}
 	}
 
-	// Compute the islands of kinematic particles
-	ComputeKinematicNodeIslands();
-
-	// Allocate the tether batches, each node has a maximum of MaxNumAttachments slots
-	TetherSlots.SetNum(DynamicNodes.Num());
-
-	for (TArray<FTether>& TetherSlot : TetherSlots)
+	if (KinematicNodes.Num())
 	{
-		TetherSlot.Reserve(MaxNumAttachments);
-	}
+		// Compute the islands of kinematic particles
+		ComputeKinematicNodeIslands();
 
-	// Find the tethers
-	if (!bUseGeodesicDistance)
-	{
-		 GenerateEuclideanTethers(Points);
-	}
-	else
-	{
-		 GenerateGeodesicTethers(Points, MaxDistances);
-	}
+		// Allocate the tether batches, each node has a maximum of MaxNumAttachments slots
+		TetherSlots.SetNum(DynamicNodes.Num());
 
-	// Update counts
-	TetherNums.Reserve(MaxNumAttachments);
-	for (const TArray<FTether>& TetherSlot : TetherSlots)
-	{
-		// Resize the count array whenever needed
-		const int32 NumUsedSlots = TetherSlot.Num();
-		TetherNums.SetNum(FMath::Max(TetherNums.Num(), NumUsedSlots));
+		for (TArray<FTether>& TetherSlot : TetherSlots)
+		{
+			TetherSlot.Reserve(MaxNumAttachments);
+		}
+
+		// Find the tethers
+		if (!bUseGeodesicDistance)
+		{
+			GenerateEuclideanTethers(Points);
+		}
+		else
+		{
+			GenerateGeodesicTethers(Points, MaxDistances);
+		}
 
-		// Increment each used slot count
-		for (int32 UsedSlot = 0; UsedSlot < NumUsedSlots; ++UsedSlot)
+		// Update counts
+		TetherNums.Reserve(MaxNumAttachments);
+		for (const TArray<FTether>& TetherSlot : TetherSlots)
 		{
-			++TetherNums[UsedSlot];
+			// Resize the count array whenever needed
+			const int32 NumUsedSlots = TetherSlot.Num();
+			TetherNums.SetNum(FMath::Max(TetherNums.Num(), NumUsedSlots));
+
+			// Increment each used slot count
+			for (int32 UsedSlot = 0; UsedSlot < NumUsedSlots; ++UsedSlot)
+			{
+				++TetherNums[UsedSlot];
+			}
 		}
 	}
 }
