Cloth wrap deformer:
- use face-average vertex normals when computing wrap weights
- check cubic function extrema for near-double-roots

#rnx
#pf 611fe2b65e7372000136a47e

diff --git a/Engine/Source/Runtime/ClothingSystemRuntimeCommon/Private/ClothPhysicalMeshData.cpp b/Engine/Source/Runtime/ClothingSystemRuntimeCommon/Private/ClothPhysicalMeshData.cpp
--- a/Engine/Source/Runtime/ClothingSystemRuntimeCommon/Private/ClothPhysicalMeshData.cpp
+++ b/Engine/Source/Runtime/ClothingSystemRuntimeCommon/Private/ClothPhysicalMeshData.cpp
@@ -268,3 +268,41 @@ void FClothPhysicalMeshData::CalculateTethers(bool bUseEuclideanDistance, bool b
 		}
 	}
 }
+
+
+void FClothPhysicalMeshData::ComputeFaceAveragedVertexNormals(TArray<FVector3f>& OutNormals) const
+{
+	OutNormals.Init(FVector3f{ 0.f, 0.f, 0.f }, Vertices.Num());
+
+	const int32 NumTris = Indices.Num() / 3;
+	for (int32 TID = 0; TID < NumTris; ++TID)
+	{
+		const uint32 VA = Indices[3 * TID + 0];
+		const uint32 VB = Indices[3 * TID + 1];
+		const uint32 VC = Indices[3 * TID + 2];
+		const FVector3f& PA = Vertices[VA];
+		const FVector3f& PB = Vertices[VB];
+		const FVector3f& PC = Vertices[VC];
+
+		FVector3f Normal = FVector3f::CrossProduct(PC - PA, PB - PA);
+		if (!Normal.Normalize())
+		{
+			// skip contributions from degenerate triangles
+			continue;
+		}
+
+		OutNormals[VA] += Normal;
+		OutNormals[VB] += Normal;
+		OutNormals[VC] += Normal;
+	}
+
+	for (int32 VertexID = 0; VertexID < OutNormals.Num(); ++VertexID )
+	{
+		FVector3f& N = OutNormals[VertexID];
+		if (!N.Normalize())
+		{
+			N = FVector3f::XAxisVector;
+		}
+	}
+}
+
diff --git a/Engine/Source/Runtime/ClothingSystemRuntimeCommon/Private/ClothingAsset.cpp b/Engine/Source/Runtime/ClothingSystemRuntimeCommon/Private/ClothingAsset.cpp
--- a/Engine/Source/Runtime/ClothingSystemRuntimeCommon/Private/ClothingAsset.cpp
+++ b/Engine/Source/Runtime/ClothingSystemRuntimeCommon/Private/ClothingAsset.cpp
@@ -323,11 +323,14 @@ bool UClothingAssetCommon::BindToSkeletalMesh(
 		}
 	}
 
-	ClothingMeshUtils::ClothMeshDesc TargetMesh(RenderPositions, RenderNormals, RenderIndices);
+	const ClothingMeshUtils::ClothMeshDesc TargetMesh(RenderPositions, RenderNormals, RenderIndices);
 
-	ClothingMeshUtils::ClothMeshDesc SourceMesh(
+	TArray<FVector3f> RecomputedVertexNormals;
+	ClothLodData.PhysicalMeshData.ComputeFaceAveragedVertexNormals(RecomputedVertexNormals);
+
+	const ClothingMeshUtils::ClothMeshDesc SourceMesh(
 		ClothLodData.PhysicalMeshData.Vertices, 
-		ClothLodData.PhysicalMeshData.Normals, 
+		RecomputedVertexNormals,
 		ClothLodData.PhysicalMeshData.Indices);
 
 	TArray<float> MaxEdgeLength;
diff --git a/Engine/Source/Runtime/ClothingSystemRuntimeCommon/Private/Utils/ClothingMeshUtils.cpp b/Engine/Source/Runtime/ClothingSystemRuntimeCommon/Private/Utils/ClothingMeshUtils.cpp
--- a/Engine/Source/Runtime/ClothingSystemRuntimeCommon/Private/Utils/ClothingMeshUtils.cpp
+++ b/Engine/Source/Runtime/ClothingSystemRuntimeCommon/Private/Utils/ClothingMeshUtils.cpp
@@ -288,6 +288,7 @@ namespace ClothingMeshUtils
 		}
 
 
+
 		bool SingleSkinningDataForVertex(const FVector3f& VertPosition,
 										 const FVector3f& VertNormal,
 										 const FVector3f& VertTangent,
@@ -431,7 +432,6 @@ namespace ClothingMeshUtils
 
 	}		// unnamed namespace
 
-
 	void GenerateMeshToMeshSkinningData(TArray<FMeshToMeshVertData>& OutSkinningData,
 										const ClothMeshDesc& TargetMesh,
 										const TArray<FVector3f>* TargetTangents,
@@ -645,14 +645,14 @@ namespace ClothingMeshUtils
 	// Solve the equation x^2 + Ax + B = 0 for real roots. 
 	// Requires an array of size 2 for the results. The return value is the number of results,
 	// either 0 or 2.
-	int32 QuadraticRoots(double Result[], double A, double B)
+	int32 QuadraticRoots(double A, double B, double Result[])
 	{
 		double D = 0.25 * A * A - B;
 		if (D >= 0.0)
 		{
 			D = FMath::Sqrt(D);
-			Result[0] = -0.5 * A + D; 
-			Result[1] = -0.5 * A - D;
+			Result[0] = -0.5 * A - D;
+			Result[1] = -0.5 * A + D; 
 			return 2; 
 		}
 		return 0;
@@ -661,7 +661,7 @@ namespace ClothingMeshUtils
 	// Solve the equation x^3 + Ax^2 + Bx + C = 0 for real roots. Requires an array of size 3 
 	// for the results. The return value is the number of results, ranging from 1 to 3.
 	// Using Viete's trig formula. See: https://en.wikipedia.org/wiki/Cubic_equation
-	int32 CubicRoots(double Result[], double A, double B, double C)
+	int32 CubicRoots(double A, double B, double C, double Result[])
 	{
 		double A2 = A * A;
 		double P = (A2 - 3.0 * B) / 9.0;
@@ -708,10 +708,68 @@ namespace ClothingMeshUtils
 		}
 	}
 
+	// Solve Ax^2 + Bx + C = 0 for real roots. Handle cases where coefficients are zero.
+	static int32 QuadraticSolve(const double Coeffs[3], double Out[])
+	{
+		const double A = Coeffs[0], B = Coeffs[1], C = Coeffs[2];
+		if (FMath::IsNearlyZero(A))
+		{
+			// First coefficient is zero, so this is at most linear
+
+			if (FMath::IsNearlyZero(B))
+			{
+				// Second coefficient is also zero, uh-oh
+				return 0;
+			}
+			else
+			{
+				// Linear Bx + C = 0 and B != 0.
+				//     =>  x = -C/B
+				Out[0] = -C / B;
+				return 1;
+			}
+		}
+		else
+		{
+			// Quadratic Ax^2 + Bx + C = 0 and A != 0
+			//        =>  x^2 + (B/A)x + C/A = 0
+
+			return QuadraticRoots(B / A, C / A, Out);
+		}
+	}
+
+	// Solve Ax^3 + Bx^2 + Cx + D = 0 for real roots. Handle cases where coefficients are zero.
+	static int32 CubicSolve(const double Coeffs[4], double Out[])
+	{
+		const double A = Coeffs[0], B = Coeffs[1], C = Coeffs[2], D = Coeffs[3];
+
+		if (FMath::IsNearlyZero(A))
+		{
+			// Leading coefficient is zero, so this is at most quadratic
+			double QCoeffs[3] = { B, C, D };
+			return QuadraticSolve(QCoeffs, Out);
+		}
+		else
+		{
+			// The cubic solver operates on cubics of the form: x^3 + Rx^2 + Sx + T = 0
+			return CubicRoots(B/A, C/A, D/A, Out);
+		}
+	}
+
+	// Find mins and maxes of cubic function Ax^3 + Bx^2 + Cx + D
+	// i.e. solve       3Ax^2 + 2Bx + C = 0
+	static int32 CubicExtrema(double A, double B, double C, double Out[])
+	{
+		double Coeffs[3] = { 3.0 * A, 2.0 * B, C };
+		return QuadraticSolve(Coeffs, Out);
+	}
+
+	constexpr int MaxNumRootsAndExtrema = 5;
+
 	static int32 CoplanarityParam(
-	const FVector3f& A, const FVector3f& B, const FVector3f& C,
-	   const FVector3f& OffsetA, const FVector3f& OffsetB, const FVector3f& OffsetC,
-	   const FVector3f& Point, double Out[3])
+		const FVector3f& A, const FVector3f& B, const FVector3f& C,
+		const FVector3f& OffsetA, const FVector3f& OffsetB, const FVector3f& OffsetC,
+		const FVector3f& Point, double Out[MaxNumRootsAndExtrema])
 	{
 		FVector3f PA = A - Point;
 		FVector3f PB = B - Point;
@@ -725,7 +783,8 @@ namespace ClothingMeshUtils
 			};
 
 		// Solve cubic A*w^3 + B*w^2 + C*w + D
-		if (FMath::IsNearlyZero(Coeffs[0], double(KINDA_SMALL_NUMBER)))
+
+		if (FMath::IsNearlyZero(Coeffs[3], double(SMALL_NUMBER)))
 		{
 			// In this case, the tetrahedron formed above is probably already at zero volume,
 			// which means the point is coplanar to the triangle without normal offsets.
@@ -736,53 +795,100 @@ namespace ClothingMeshUtils
 		}
 		else
 		{
-			for (int32 I = 1; I < 4; I++)
+			// Search for "near double roots": points where the first derivative is zero and where the function value
+			// is close to zero.
+			double Extrema[2];
+			const double CloseToZeroTolerance = 0.1 * (FMath::Abs(Coeffs[0]) + FMath::Abs(Coeffs[1]) + FMath::Abs(Coeffs[2]) + FMath::Abs(Coeffs[3]));
+			const int32 NumExtrema = CubicExtrema(Coeffs[0], Coeffs[1], Coeffs[2], Extrema);
+
+			int32 NumExtremaUsed = 0;
+			for (int32 I = 0; I < NumExtrema; ++I)
 			{
-				Coeffs[I] /= Coeffs[0];
+				const double X = Extrema[I];
+				const double FuncValue = Coeffs[0] * X * X * X + Coeffs[1] * X * X + Coeffs[2] * X + Coeffs[3];
+				if (FMath::IsNearlyZero(FuncValue, CloseToZeroTolerance))
+				{
+					Out[NumExtremaUsed++] = Extrema[I];
+				}
 			}
-			
-			return CubicRoots(Out, Coeffs[1], Coeffs[2], Coeffs[3]);
+
+			// Now find roots of the cubic function
+			double Roots[3];
+			const int32 NumRoots = CubicSolve(Coeffs, Roots);
+
+			// Combine roots and extrema
+			for (int32 I = 0; I < NumRoots; ++I)
+			{
+				int NextI = NumExtremaUsed + I;
+				check(NextI < MaxNumRootsAndExtrema);
+				Out[NextI] = Roots[I];
+			}
+
+			// Sort by magnitude of roots/extrema to bias towards eventually choosing smaller parameter values
+			Algo::Sort(TArrayView<double>(Out, NumExtremaUsed + NumRoots), [](const double& A, const double& B) {
+				return FMath::Abs(A) < FMath::Abs(B);
+			});
+
+			return NumExtremaUsed + NumRoots;
 		}
 	}
 
 	FVector4 GetPointBaryAndDistWithNormals(
 		const FVector3f& A, const FVector3f& B, const FVector3f& C,
-		const FVector3f& NA, const FVector3f& NB, const FVector3f& NC,
+		const FVector3f& InputNA, const FVector3f& InputNB, const FVector3f& InputNC,
 		const FVector3f& Point)
 	{
+		// Input normals are inverse of what they are in the shader code. Flip the sign when computing weights that
+		// will be used in the shader code.
+		const FVector3f UseNA = -InputNA;
+		const FVector3f UseNB = -InputNB;
+		const FVector3f UseNC = -InputNC;
+
 		// Adapted from cloth CCD paper [Bridson et al. 2002]
 		// First find W such that Point lies in the plane defined by {A+wNA, B+wNB, C+wNC}
-		// Pass in inverted normals, since they get inverted at runtime (Left handed system).
-		double W[3];
-		const int32 Count = CoplanarityParam(A, B, C, NA, NB, NC, Point, W);
+		double W[MaxNumRootsAndExtrema];
+		const int32 CoplanarityParamCount = CoplanarityParam(A, B, C, UseNA, UseNB, UseNC, Point, W);
 		
-		// We should always have a solution. The only case where we can get zero roots from the
-		// cubic equation, we handle separately.
-		checkSlow(Count != 0);
+		if (CoplanarityParamCount == 0)
+		{
+			// Found no parameter w for which Point \in span{A+wNA, B+wNB, C+wNC}
+			// Fall back to using the triangle normal
+			return GetPointBaryAndDist(A, B, C, Point);
+		}
 
 		FVector4 BaryAndDist;
-		double MinDistanceSq = std::numeric_limits<double>::max();
+
+		float MinDistanceSq = TNumericLimits<float>::Max();
+		bool bAnySolutionFound = false;
 
 		// If the solution gives us barycentric coordinates that lie purely within the triangle,
 		// then choose that. Otherwise try to minimize the distance of the projected point to
 		// be as close to the triangle as possible.
-		for (int32 Index = 0; Index < Count; Index++)
+		for (int32 CoplanarityParamIndex = 0; CoplanarityParamIndex < CoplanarityParamCount; ++CoplanarityParamIndex)
 		{
 			// Then find the barycentric coordinates of Point wrt {A+wNA, B+wNB, C+wNC}
-			FVector AW = A + W[Index] * NA;
-			FVector BW = B + W[Index] * NB;
-			FVector CW = C + W[Index] * NC;
-
+			FVector3f AW = A + W[CoplanarityParamIndex] * UseNA;
+			FVector3f BW = B + W[CoplanarityParamIndex] * UseNB;
+			FVector3f CW = C + W[CoplanarityParamIndex] * UseNC;
+		
 			FPlane4f TrianglePlane(AW, BW, CW);
 
 			const FVector3f PointOnTriPlane = FVector::PointPlaneProject(Point, TrianglePlane);
 			const FVector3f BaryCoords = FMath::ComputeBaryCentric2D(PointOnTriPlane, AW, BW, CW);
 
+			if (BaryCoords.X == BaryCoords.Y && BaryCoords.Y == BaryCoords.Z && BaryCoords.Z == 0.0f)
+			{
+				// Degenerate triangle at this value of W
+				continue;
+			}
+
+			bAnySolutionFound = true;
+
 			if (BaryCoords.X >= 0.0f && BaryCoords.X <= 1.0f &&
 				BaryCoords.Y >= 0.0f && BaryCoords.Y <= 1.0f &&
 				BaryCoords.Z >= 0.0f && BaryCoords.Z <= 1.0f)
 			{
-				BaryAndDist = FVector4(BaryCoords, -W[Index]);
+				BaryAndDist = FVector4(BaryCoords, W[CoplanarityParamIndex]);
 				break;
 			}
 			
@@ -792,23 +898,34 @@ namespace ClothingMeshUtils
 			
 			if (DistSq < MinDistanceSq)
 			{
-				BaryAndDist = FVector4(BaryCoords, -W[Index]);
+				BaryAndDist = FVector4(BaryCoords, W[CoplanarityParamIndex]);
 				MinDistanceSq = DistSq;
 			}
 		}
 
-		const FVector3f ReprojectedPoint =
-			BaryAndDist.X * (A - NA * BaryAndDist.W) +
-			BaryAndDist.Y * (B - NB * BaryAndDist.W) +
-			BaryAndDist.Z * (C - NC * BaryAndDist.W);
 
-		const float Distance = FVector::Distance(Point, ReprojectedPoint);
+		bool bRequireFallback = false;
 
-		// Check if the reprojected point is far from the original. If it is, fall back on
-		// the old method of computing the bary values.
-		// FIXME: Should we test other cage triangles instead? It's possible that
-		// GetBestTriangleBaseIndex is not actually picking the /best/ one.
-		const bool bRequireFallback = !FMath::IsNearlyZero(Distance, KINDA_SMALL_NUMBER); 
+		if (bAnySolutionFound)
+		{
+			const FVector3f ReprojectedPoint =
+				BaryAndDist.X * (A + UseNA * BaryAndDist.W) +
+				BaryAndDist.Y * (B + UseNB * BaryAndDist.W) +
+				BaryAndDist.Z * (C + UseNC * BaryAndDist.W);
+
+			const float Distance = FVector::Distance(Point, ReprojectedPoint);
+
+			// Check if the reprojected point is far from the original. If it is, fall back on
+			// the old method of computing the bary values.
+			// FIXME: Should we test other cage triangles instead? It's possible that
+			// GetBestTriangleBaseIndex is not actually picking the /best/ one.
+			bRequireFallback = !FMath::IsNearlyZero(Distance, KINDA_SMALL_NUMBER);
+		}
+		else
+		{
+			// no viable solution
+			bRequireFallback = true;
+		}
 
 		if (bRequireFallback)
 		{
@@ -961,6 +1078,7 @@ namespace ClothingMeshUtils
 		}
 		return TArray<int32>();
 	}
+
 }
 
 #undef LOCTEXT_NAMESPACE
diff --git a/Engine/Source/Runtime/ClothingSystemRuntimeCommon/Public/ClothPhysicalMeshData.h b/Engine/Source/Runtime/ClothingSystemRuntimeCommon/Public/ClothPhysicalMeshData.h
--- a/Engine/Source/Runtime/ClothingSystemRuntimeCommon/Public/ClothPhysicalMeshData.h
+++ b/Engine/Source/Runtime/ClothingSystemRuntimeCommon/Public/ClothPhysicalMeshData.h
@@ -49,6 +49,10 @@ struct CLOTHINGSYSTEMRUNTIMECOMMON_API FClothPhysicalMeshData
 	/** Recalculate the long range attachment tethers. */
 	void CalculateTethers(bool bUseEuclideanDistance, bool bUseGeodesicDistance);
 
+	/** Compute vertex normals as unweighted average of incident face normals. If all incident triangles are degenerate,
+	    the vertex normal is assigned a value of FVector3f::XAxisVector. */
+	void ComputeFaceAveragedVertexNormals(TArray<FVector3f>& OutNormals) const;
+
 	/** Retrieve whether a vertex weight array has already been registered. */
 	template<typename T>
 	bool HasWeightMap(const T Target) const { return WeightMaps.Contains((uint32)Target); }
