Cloth wrap deformer: test cubic solver roots vs distance to the original triangles, so as to avoid large roots.

#rnx

diff --git a/Engine/Source/Runtime/ClothingSystemRuntimeCommon/Private/Utils/ClothingMeshUtils.cpp b/Engine/Source/Runtime/ClothingSystemRuntimeCommon/Private/Utils/ClothingMeshUtils.cpp
--- a/Engine/Source/Runtime/ClothingSystemRuntimeCommon/Private/Utils/ClothingMeshUtils.cpp
+++ b/Engine/Source/Runtime/ClothingSystemRuntimeCommon/Private/Utils/ClothingMeshUtils.cpp
@@ -856,8 +856,10 @@ namespace ClothingMeshUtils
 			return GetPointBaryAndDist(A, B, C, Point);
 		}
 
+		const FVector3f ClosestPoint = FMath::ClosestPointOnTriangleToPoint(Point, A, B, C);
+		const float DistanceToTriangle = FVector3f::Distance(Point, ClosestPoint);
+		
 		FVector4 BaryAndDist;
-
 		float MinDistanceSq = TNumericLimits<float>::Max();
 		bool bAnySolutionFound = false;
 
@@ -866,6 +868,11 @@ namespace ClothingMeshUtils
 		// be as close to the triangle as possible.
 		for (int32 CoplanarityParamIndex = 0; CoplanarityParamIndex < CoplanarityParamCount; ++CoplanarityParamIndex)
 		{
+			if (FMath::Abs(W[CoplanarityParamIndex]) > 3.0f * DistanceToTriangle)
+			{
+				continue;
+			}
+
 			// Then find the barycentric coordinates of Point wrt {A+wNA, B+wNB, C+wNC}
 			FVector3f AW = A + W[CoplanarityParamIndex] * UseNA;
 			FVector3f BW = B + W[CoplanarityParamIndex] * UseNB;
