Chaos Cloth - Fix max distance kinematic threshold issues detected in debug builds.

diff --git a/Engine/Source/Runtime/ClothingSystemRuntimeCommon/Private/ClothTetherData.cpp b/Engine/Source/Runtime/ClothingSystemRuntimeCommon/Private/ClothTetherData.cpp
--- a/Engine/Source/Runtime/ClothingSystemRuntimeCommon/Private/ClothTetherData.cpp
+++ b/Engine/Source/Runtime/ClothingSystemRuntimeCommon/Private/ClothTetherData.cpp
@@ -6,6 +6,7 @@ class FClothTetherDataPrivate
 {
 public:
 	static constexpr int32 MaxNumAttachments = 4;  // Max recommended number of tethers per point, could eventually been specified by config
+	static constexpr float KinematicDistanceThreshold = 0.1f;  // Particles which are barely allowed to move around should be considered kinematic
 
 	typedef TTuple<int32, int32, float> FTether;
 
@@ -60,7 +61,7 @@ void FClothTetherData::GenerateTethers(
 	const TArray<TArray<FClothTetherDataPrivate::FTether>>& TetherSlots = ClothTetherData.GetTetherSlots();
 	const TArray<int32> TetherNums = ClothTetherData.GetTetherNums();
 
-	// Reorganize the multiple tethers per node array into a single sequencial batches of tethers more suitable for parallel processing
+	// Reorganize the multiple tethers per node array into a single sequential batches of tethers more suitable for parallel processing
 	const int32 NumDynamicNodes = TetherSlots.Num();
 	const int32 MaxNumUsedSlots = TetherNums.Num();
 	Tethers.Reset(MaxNumUsedSlots);
@@ -124,7 +125,7 @@ FClothTetherDataPrivate::FClothTetherDataPrivate(
 	KinematicNodes.Reserve(Nodes.Num() / 3);  // Start at 33% of the number of nodes to minimize this array's reallocations
 	for (const int32 Node : Nodes)
 	{
-		if (MaxDistances[Node])
+		if (MaxDistances[Node] >= KinematicDistanceThreshold)
 		{
 			DynamicNodes.Add(Node);
 		}
@@ -336,7 +337,7 @@ void FClothTetherDataPrivate::GenerateGeodesicTethers(const TConstArrayView<FVec
 			const TSet<int32>& Neighbors = NodeToNeighbors[KinematicNode];
 			for (const int32 Neighbor : Neighbors)
 			{
-				if (MaxDistances[Neighbor])
+				if (MaxDistances[Neighbor] >= KinematicDistanceThreshold)
 				{
 					SeedIslands[IslandIndex].Add(KinematicNode);
 					++NumSeeds;
@@ -415,7 +416,7 @@ void FClothTetherDataPrivate::GenerateGeodesicTethers(const TConstArrayView<FVec
 				check(NeighborNode != ParentNode);
 
 				// Do not progress onto kinematic nodes
-				if (!MaxDistances[NeighborNode])
+				if (MaxDistances[NeighborNode] < KinematicDistanceThreshold)
 				{
 					continue;
 				}
diff --git a/Engine/Source/Runtime/ClothingSystemRuntimeCommon/Private/ClothingAsset.cpp b/Engine/Source/Runtime/ClothingSystemRuntimeCommon/Private/ClothingAsset.cpp
--- a/Engine/Source/Runtime/ClothingSystemRuntimeCommon/Private/ClothingAsset.cpp
+++ b/Engine/Source/Runtime/ClothingSystemRuntimeCommon/Private/ClothingAsset.cpp
@@ -955,12 +955,15 @@ void UClothingAssetCommon::PostLoad()
 
 	// Cache tethers and reference bone index should already have been calculated
 	const int32 FrostyMainBranchObjectVersion = GetLinkerCustomVersion(FUE5PrivateFrostyStreamObjectVersion::GUID);
-	if (FrostyMainBranchObjectVersion < FUE5PrivateFrostyStreamObjectVersion::ChaosClothAddTethersToCachedData)
+	if (FrostyMainBranchObjectVersion < FUE5PrivateFrostyStreamObjectVersion::ChaosClothRemoveKinematicTethers)
 	{
 		ClothingCachedDataFlags |= EClothingCachedDataFlagsCommon::Tethers;
 
-		// ReferenceBoneIndex is only required when rebinding the cloth.
-		CalculateReferenceBoneIndex();
+		if (FrostyMainBranchObjectVersion < FUE5PrivateFrostyStreamObjectVersion::ChaosClothAddTethersToCachedData)
+		{
+			// ReferenceBoneIndex is only required when rebinding the cloth.
+			CalculateReferenceBoneIndex();
+		}
 	}
 
 	// After fixing the content, we are ready to call functions that rely on it
diff --git a/Engine/Source/Runtime/Core/Public/UObject/UE5PrivateFrostyStreamObjectVersion.h b/Engine/Source/Runtime/Core/Public/UObject/UE5PrivateFrostyStreamObjectVersion.h
--- a/Engine/Source/Runtime/Core/Public/UObject/UE5PrivateFrostyStreamObjectVersion.h
+++ b/Engine/Source/Runtime/Core/Public/UObject/UE5PrivateFrostyStreamObjectVersion.h
@@ -27,6 +27,9 @@ struct CORE_API FUE5PrivateFrostyStreamObjectVersion
 		// Changed world partition HLODs cells from FSotObjectPath to FName
 		ConvertWorldPartitionHLODsCellsToName,
 
+		// Re-calculate the long range attachment to prevent kinematic tethers.
+		ChaosClothRemoveKinematicTethers,
+
 		// -----<new versions can be added above this line>-------------------------------------------------
 		VersionPlusOne,
 		LatestVersion = VersionPlusOne - 1
