Cloth - Fix the Backstop normals by reverting and iterating on some of the recent LOD Bias changes.

diff --git a/Engine/Plugins/Experimental/ChaosCloth/Source/Chaos/Private/ChaosCloth/ChaosClothingSimulation.cpp b/Engine/Plugins/Experimental/ChaosCloth/Source/Chaos/Private/ChaosCloth/ChaosClothingSimulation.cpp
--- a/Engine/Plugins/Experimental/ChaosCloth/Source/Chaos/Private/ChaosCloth/ChaosClothingSimulation.cpp
+++ b/Engine/Plugins/Experimental/ChaosCloth/Source/Chaos/Private/ChaosCloth/ChaosClothingSimulation.cpp
@@ -84,8 +84,8 @@ namespace ChaosClothingSimulationCVar
 	TAutoConsoleVariable<bool> DebugDrawGravity             (TEXT("p.ChaosCloth.DebugDrawGravity"             ), false, TEXT("Whether to debug draw the Chaos Cloth gravity acceleration vector"), ECVF_Cheat);
 	TAutoConsoleVariable<bool> DebugDrawPhysMeshWired       (TEXT("p.ChaosCloth.DebugDrawPhysMeshWired"       ), false, TEXT("Whether to debug draw the Chaos Cloth wireframe meshes"), ECVF_Cheat);
 	TAutoConsoleVariable<bool> DebugDrawAnimMeshWired       (TEXT("p.ChaosCloth.DebugDrawAnimMeshWired"       ), false, TEXT("Whether to debug draw the animated/kinematic Cloth wireframe meshes"), ECVF_Cheat);
+	TAutoConsoleVariable<bool> DebugDrawAnimNormals         (TEXT("p.ChaosCloth.DebugDrawAmimNormals"         ), false, TEXT("Whether to debug draw the animated/kinematic Cloth normals"), ECVF_Cheat);
 	TAutoConsoleVariable<bool> DebugDrawPointNormals        (TEXT("p.ChaosCloth.DebugDrawPointNormals"        ), false, TEXT("Whether to debug draw the Chaos Cloth point normals"), ECVF_Cheat);
-	TAutoConsoleVariable<bool> DebugDrawInversedPointNormals(TEXT("p.ChaosCloth.DebugDrawInversedPointNormals"), false, TEXT("Whether to debug draw the Chaos Cloth inversed point normals"), ECVF_Cheat);
 	TAutoConsoleVariable<bool> DebugDrawFaceNormals         (TEXT("p.ChaosCloth.DebugDrawFaceNormals"         ), false, TEXT("Whether to debug draw the Chaos Cloth face normals"), ECVF_Cheat);
 	TAutoConsoleVariable<bool> DebugDrawInversedFaceNormals (TEXT("p.ChaosCloth.DebugDrawInversedFaceNormals" ), false, TEXT("Whether to debug draw the Chaos Cloth inversed face normals"), ECVF_Cheat);
 	TAutoConsoleVariable<bool> DebugDrawCollision           (TEXT("p.ChaosCloth.DebugDrawCollision"           ), false, TEXT("Whether to debug draw the Chaos Cloth collisions"), ECVF_Cheat);
@@ -568,8 +568,8 @@ void FClothingSimulation::Simulate(IClothingSimulationContext* InContext)
 	if (ChaosClothingSimulationCVar::DebugDrawGravity             .GetValueOnAnyThread()) { DebugDrawGravity             (); }
 	if (ChaosClothingSimulationCVar::DebugDrawPhysMeshWired       .GetValueOnAnyThread()) { DebugDrawPhysMeshWired       (); }
 	if (ChaosClothingSimulationCVar::DebugDrawAnimMeshWired       .GetValueOnAnyThread()) { DebugDrawAnimMeshWired       (); }
+	if (ChaosClothingSimulationCVar::DebugDrawAnimNormals         .GetValueOnAnyThread()) { DebugDrawAnimNormals         (); }
 	if (ChaosClothingSimulationCVar::DebugDrawPointNormals        .GetValueOnAnyThread()) { DebugDrawPointNormals        (); }
-	if (ChaosClothingSimulationCVar::DebugDrawInversedPointNormals.GetValueOnAnyThread()) { DebugDrawInversedPointNormals(); }
 	if (ChaosClothingSimulationCVar::DebugDrawCollision           .GetValueOnAnyThread()) { DebugDrawCollision           (); }
 	if (ChaosClothingSimulationCVar::DebugDrawBackstops           .GetValueOnAnyThread()) { DebugDrawBackstops           (); }
 	if (ChaosClothingSimulationCVar::DebugDrawBackstopDistances   .GetValueOnAnyThread()) { DebugDrawBackstopDistances   (); }
@@ -1346,6 +1346,35 @@ void FClothingSimulation::DebugDrawAnimMeshWired(FPrimitiveDrawInterface* PDI) c
 	}
 }
 
+void FClothingSimulation::DebugDrawAnimNormals(FPrimitiveDrawInterface* PDI) const
+{
+	check(Solver);
+	static const FLinearColor KinematicColor = FColor::Magenta;
+
+	const FVec3& LocalSpaceLocation = Solver->GetLocalSpaceLocation();
+
+	for (const FClothingSimulationCloth* const Cloth : Solver->GetCloths())
+	{
+		const int32 Offset = Cloth->GetOffset(Solver.Get());
+		if (Offset == INDEX_NONE)
+		{
+			continue;
+		}
+
+		const TConstArrayView<FVec3> Positions = Cloth->GetAnimationPositions(Solver.Get());
+		const TConstArrayView<FVec3> Normals = Cloth->GetAnimationNormals(Solver.Get());
+		check(Normals.Num() == Positions.Num());
+
+		for (int32 Index = 0; Index < Positions.Num(); ++Index)
+		{
+			const FVector Pos0 = LocalSpaceLocation + Positions[Index];
+			const FVector Pos1 = Pos0 + Normals[Index] * 20.f;
+
+			DrawLine(PDI, Pos0, Pos1, KinematicColor);
+		}
+	}
+}
+
 void FClothingSimulation::DebugDrawPointNormals(FPrimitiveDrawInterface* PDI) const
 {
 	check(Solver);
diff --git a/Engine/Plugins/Experimental/ChaosCloth/Source/Chaos/Public/ChaosCloth/ChaosClothingSimulation.h b/Engine/Plugins/Experimental/ChaosCloth/Source/Chaos/Public/ChaosCloth/ChaosClothingSimulation.h
--- a/Engine/Plugins/Experimental/ChaosCloth/Source/Chaos/Public/ChaosCloth/ChaosClothingSimulation.h
+++ b/Engine/Plugins/Experimental/ChaosCloth/Source/Chaos/Public/ChaosCloth/ChaosClothingSimulation.h
@@ -105,7 +105,9 @@ namespace Chaos
 		// Editor & runtime debug draw functions
 		CHAOSCLOTH_API void DebugDrawPhysMeshWired(FPrimitiveDrawInterface* PDI = nullptr) const;
 		CHAOSCLOTH_API void DebugDrawAnimMeshWired(FPrimitiveDrawInterface* PDI = nullptr) const;
+		CHAOSCLOTH_API void DebugDrawAnimNormals(FPrimitiveDrawInterface* PDI = nullptr) const;
 		CHAOSCLOTH_API void DebugDrawPointNormals(FPrimitiveDrawInterface* PDI = nullptr) const;
+		UE_DEPRECATED(5.0, "DebugDrawInversedPointNormals is mostly redundant and will be removed, use DebugDrawPointNormals instead")
 		CHAOSCLOTH_API void DebugDrawInversedPointNormals(FPrimitiveDrawInterface* PDI = nullptr) const;
 		CHAOSCLOTH_API void DebugDrawCollision(FPrimitiveDrawInterface* PDI = nullptr) const;
 		CHAOSCLOTH_API void DebugDrawBackstops(FPrimitiveDrawInterface* PDI = nullptr) const;
diff --git a/Engine/Plugins/Experimental/ChaosClothEditor/Source/Chaos/Private/ChaosClothEditor/ChaosSimulationEditorExtender.cpp b/Engine/Plugins/Experimental/ChaosClothEditor/Source/Chaos/Private/ChaosClothEditor/ChaosSimulationEditorExtender.cpp
--- a/Engine/Plugins/Experimental/ChaosClothEditor/Source/Chaos/Private/ChaosClothEditor/ChaosSimulationEditorExtender.cpp
+++ b/Engine/Plugins/Experimental/ChaosClothEditor/Source/Chaos/Private/ChaosClothEditor/ChaosSimulationEditorExtender.cpp
@@ -72,8 +72,8 @@ const FVisualizationOption FVisualizationOption::OptionData[] =
 	FVisualizationOption(&Chaos::FClothingSimulation::DebugDrawAnimMeshWired       , TEXT("p.ChaosClothEditor.DebugDrawAnimMeshWired"       ), LOCTEXT("ChaosVisName_AnimMeshWire"        , "Animated Mesh (Wireframe)"  ), LOCTEXT("ChaosVisName_AnimMeshWired_ToolTip"       , "Draws the current animated mesh input in wireframe")),
 	FVisualizationOption(&Chaos::FClothingSimulation::DebugDrawParticleIndices     , TEXT("p.ChaosClothEditor.DebugDrawParticleIndices"     ), LOCTEXT("ChaosVisName_ParticleIndices"     , "Particle Indices"           ), LOCTEXT("ChaosVisName_ParticleIndices_ToolTip"     , "Draws the particle indices as instantiated by the solver")),
 	FVisualizationOption(&Chaos::FClothingSimulation::DebugDrawElementIndices      , TEXT("p.ChaosClothEditor.DebugDrawElementIndices"      ), LOCTEXT("ChaosVisName_ElementIndices"      , "Element Indices"            ), LOCTEXT("ChaosVisName_ElementIndices_ToolTip"      , "Draws the element's (triangle or other) indices as instantiated by the solver")),
-	FVisualizationOption(&Chaos::FClothingSimulation::DebugDrawPointNormals        , TEXT("p.ChaosClothEditor.DebugDrawPointNormals"        ), LOCTEXT("ChaosVisName_PointNormals"        , "Point Normals"              ), LOCTEXT("ChaosVisName_PointNormals_ToolTip"        , "Draws the current point normals for the simulation mesh")),
-	FVisualizationOption(&Chaos::FClothingSimulation::DebugDrawInversedPointNormals, TEXT("p.ChaosClothEditor.DebugDrawInversedPointNormals"), LOCTEXT("ChaosVisName_InversedPointNormals", "Inversed Point Normals"     ), LOCTEXT("ChaosVisName_InversedPointNormals_ToolTip", "Draws the inversed point normals for the simulation mesh")),
+	FVisualizationOption(&Chaos::FClothingSimulation::DebugDrawPointNormals        , TEXT("p.ChaosClothEditor.DebugDrawPointNormals"        ), LOCTEXT("ChaosVisName_PointNormals"        , "Physical Mesh Normals"      ), LOCTEXT("ChaosVisName_PointNormals_ToolTip"        , "Draws the current point normals for the simulation mesh")),
+	FVisualizationOption(&Chaos::FClothingSimulation::DebugDrawAnimNormals         , TEXT("p.ChaosClothEditor.DebugDrawAnimNormals"         ), LOCTEXT("ChaosVisName_AnimNormals"         , "Animated Mesh Normals"      ), LOCTEXT("ChaosVisName_AnimNormals_ToolTip"         , "Draws the current point normals for the animated mesh")),
 	FVisualizationOption(&Chaos::FClothingSimulation::DebugDrawCollision           , TEXT("p.ChaosClothEditor.DebugDrawCollision"           ), LOCTEXT("ChaosVisName_Collision"           , "Collisions"                 ), LOCTEXT("ChaosVisName_Collision_ToolTip"           , "Draws the collision bodies the simulation is currently using")),
 	FVisualizationOption(&Chaos::FClothingSimulation::DebugDrawBackstops           , TEXT("p.ChaosClothEditor.DebugDrawBackstops"           ), LOCTEXT("ChaosVisName_Backstop"            , "Backstops"                  ), LOCTEXT("ChaosVisName_Backstop_ToolTip"            , "Draws the backstop radius and position for each simulation particle")),
 	FVisualizationOption(&Chaos::FClothingSimulation::DebugDrawBackstopDistances   , TEXT("p.ChaosClothEditor.DebugDrawBackstopDistances"   ), LOCTEXT("ChaosVisName_BackstopDistance"    , "Backstop Distances"         ), LOCTEXT("ChaosVisName_BackstopDistance_ToolTip"    , "Draws the backstop distance offset for each simulation particle"), /*bDisablesSimulation =*/true),
diff --git a/Engine/Source/Editor/ClothingSystemEditor/Private/ClothingAssetFactory.cpp b/Engine/Source/Editor/ClothingSystemEditor/Private/ClothingAssetFactory.cpp
--- a/Engine/Source/Editor/ClothingSystemEditor/Private/ClothingAssetFactory.cpp
+++ b/Engine/Source/Editor/ClothingSystemEditor/Private/ClothingAssetFactory.cpp
@@ -1234,6 +1234,7 @@ bool UClothingAssetFactory::ImportToLodInternal(
 		const FSoftSkinVertex& SourceVert = SourceSection.SoftVertices[OriginalIndexes[VertexIndex]];
 
 		PhysMesh.Vertices[VertexIndex] = SourceVert.Position;
+		PhysMesh.Normals[VertexIndex] = SourceVert.TangentZ;
 		PhysMesh.VertexColors[VertexIndex] = SourceVert.Color;
 
 		FClothVertBoneData& BoneData = PhysMesh.BoneData[VertexIndex];
@@ -1256,8 +1257,6 @@ bool UClothingAssetFactory::ImportToLodInternal(
 		}
 	}
 
-	PhysMesh.CalculateNormals();
-
 	// Add a max distance parameter mask to the physics mesh
 	FPointWeightMap& PhysMeshMaxDistances = PhysMesh.AddWeightMap(EWeightMapTargetCommon::MaxDistance);
 	PhysMeshMaxDistances.Initialize(PhysMesh.Vertices.Num());
diff --git a/Engine/Source/Runtime/ClothingSystemRuntimeCommon/Private/ClothPhysicalMeshData.cpp b/Engine/Source/Runtime/ClothingSystemRuntimeCommon/Private/ClothPhysicalMeshData.cpp
--- a/Engine/Source/Runtime/ClothingSystemRuntimeCommon/Private/ClothPhysicalMeshData.cpp
+++ b/Engine/Source/Runtime/ClothingSystemRuntimeCommon/Private/ClothPhysicalMeshData.cpp
@@ -305,10 +305,3 @@ void FClothPhysicalMeshData::ComputeFaceAveragedVertexNormals(TArray<FVector3f>&
 		}
 	}
 }
-
-void FClothPhysicalMeshData::CalculateNormals()
-{
-PRAGMA_DISABLE_DEPRECATION_WARNINGS
-	ComputeFaceAveragedVertexNormals(Normals);
-PRAGMA_ENABLE_DEPRECATION_WARNINGS
-}
diff --git a/Engine/Source/Runtime/ClothingSystemRuntimeCommon/Private/ClothingAsset.cpp b/Engine/Source/Runtime/ClothingSystemRuntimeCommon/Private/ClothingAsset.cpp
--- a/Engine/Source/Runtime/ClothingSystemRuntimeCommon/Private/ClothingAsset.cpp
+++ b/Engine/Source/Runtime/ClothingSystemRuntimeCommon/Private/ClothingAsset.cpp
@@ -327,11 +327,7 @@ bool UClothingAssetCommon::BindToSkeletalMesh(
 	}
 
 	const ClothingMeshUtils::ClothMeshDesc TargetMesh(RenderPositions, RenderNormals, RenderTangents, RenderIndices);
-
-	const ClothingMeshUtils::ClothMeshDesc SourceMesh(
-		ClothLodData.PhysicalMeshData.Vertices,
-		ClothLodData.PhysicalMeshData.Normals,
-		ClothLodData.PhysicalMeshData.Indices);
+	const ClothingMeshUtils::ClothMeshDesc SourceMesh(ClothLodData.PhysicalMeshData.Vertices, ClothLodData.PhysicalMeshData.Indices);  // Calculate averaged normals
 
 	const FPointWeightMap* const MaxDistances = ClothLodData.PhysicalMeshData.FindWeightMap(EWeightMapTargetCommon::MaxDistance);
 
@@ -530,10 +526,7 @@ void UClothingAssetCommon::UpdateLODBiasMappings(const USkeletalMesh* SkeletalMe
 	{
 		// Prepare the deformer source (simulation) mesh descriptor
 		const FClothLODDataCommon& ClothLodData = LodData[LodMap[UpdatedLODIndex]];
-		const ClothingMeshUtils::ClothMeshDesc SourceMesh(
-			ClothLodData.PhysicalMeshData.Vertices, 
-			ClothLodData.PhysicalMeshData.Normals,
-			ClothLodData.PhysicalMeshData.Indices);
+		const ClothingMeshUtils::ClothMeshDesc SourceMesh(ClothLodData.PhysicalMeshData.Vertices, ClothLodData.PhysicalMeshData.Indices);  // Calculate averaged normals
 
 		// Retrieve max distance mask
 		const FPointWeightMap* const MaxDistances = ClothLodData.PhysicalMeshData.FindWeightMap(EWeightMapTargetCommon::MaxDistance);
@@ -621,10 +614,7 @@ void UClothingAssetCommon::UpdateLODBiasMappings(const USkeletalMesh* SkeletalMe
 
 			// Prepare the deformer source (simulation) mesh descriptor
 			const FClothLODDataCommon& ClothLodData = ClothingAsset->LodData[LodMap[LODIndex]];
-			const ClothingMeshUtils::ClothMeshDesc SourceMesh(
-				ClothLodData.PhysicalMeshData.Vertices,
-				ClothLodData.PhysicalMeshData.Normals,
-				ClothLodData.PhysicalMeshData.Indices);
+			const ClothingMeshUtils::ClothMeshDesc SourceMesh(ClothLodData.PhysicalMeshData.Vertices, ClothLodData.PhysicalMeshData.Indices);  // Calculate averaged normals
 
 			// Retrieve max distance mask
 			const FPointWeightMap* const MaxDistances = ClothLodData.PhysicalMeshData.FindWeightMap(EWeightMapTargetCommon::MaxDistance);
@@ -1061,16 +1051,6 @@ void UClothingAssetCommon::PostLoad()
 	}
 	ClothLodData_DEPRECATED.Empty();
 
-	// Re-calulate normals
-	const int32 UE5ReleaseStreamObjectVersion = GetLinkerCustomVersion(FUE5ReleaseStreamObjectVersion::GUID);
-	if (UE5ReleaseStreamObjectVersion < FUE5ReleaseStreamObjectVersion::AddClothMappingLODBias)
-	{
-		for (FClothLODDataCommon& Lod : LodData)
-		{
-			Lod.PhysicalMeshData.CalculateNormals();
-		}
-	}
-
 	const int32 AnimPhysCustomVersion = GetLinkerCustomVersion(FAnimPhysObjectVersion::GUID);
 	if (AnimPhysCustomVersion < FAnimPhysObjectVersion::AddedClothingMaskWorkflow)
 	{
diff --git a/Engine/Source/Runtime/ClothingSystemRuntimeCommon/Private/Utils/ClothingMeshUtils.cpp b/Engine/Source/Runtime/ClothingSystemRuntimeCommon/Private/Utils/ClothingMeshUtils.cpp
--- a/Engine/Source/Runtime/ClothingSystemRuntimeCommon/Private/Utils/ClothingMeshUtils.cpp
+++ b/Engine/Source/Runtime/ClothingSystemRuntimeCommon/Private/Utils/ClothingMeshUtils.cpp
@@ -56,6 +56,39 @@ namespace ClothingMeshUtils
 		bHasValidMaxEdgeLengths = true;
 	}
 
+	void ClothMeshDesc::ComputeAveragedNormals() const
+	{
+		AveragedNormals.Init(FVector3f::ZeroVector, Positions.Num());  // Note AveragedNormals is a mutable, so can be modified in this function
+
+		const int32 NumTriangles = Indices.Num() / 3;
+		for (int32 TriangleIndex = 0; TriangleIndex < NumTriangles; ++TriangleIndex)
+		{
+			const uint32 Index0 = Indices[3 * TriangleIndex + 0];
+			const uint32 Index1 = Indices[3 * TriangleIndex + 1];
+			const uint32 Index2 = Indices[3 * TriangleIndex + 2];
+			const FVector3f& Position0 = Positions[Index0];
+			const FVector3f& Position1 = Positions[Index1];
+			const FVector3f& Position2 = Positions[Index2];
+
+			FVector3f Normal = FVector3f::CrossProduct(Position2 - Position0, Position1 - Position0);
+			if (Normal.Normalize())  // Skip contributions from degenerate triangles
+			{
+				AveragedNormals[Index0] += Normal;
+				AveragedNormals[Index1] += Normal;
+				AveragedNormals[Index2] += Normal;
+			}
+		}
+
+		for (int32 Index = 0; Index < AveragedNormals.Num(); ++Index)
+		{
+			FVector3f& Normal = AveragedNormals[Index];
+			if (!Normal.Normalize())
+			{
+				Normal = FVector3f::XAxisVector;
+			}
+		}
+	}
+
 	// Explicit template instantiations of SkinPhysicsMesh
 	// TODO: Deprecated, remove this templated version post 5.0
 	template
@@ -398,13 +431,15 @@ namespace ClothingMeshUtils
 										 int32 ClosestTriangleBaseIdx,
 										 FMeshToMeshVertData& OutSkinningData)
 		{
+			check(SourceMesh.AveragedNormals.Num() == SourceMesh.Positions.Num());
+
 			const FVector3f& A = SourceMesh.Positions[SourceMesh.Indices[ClosestTriangleBaseIdx]];
 			const FVector3f& B = SourceMesh.Positions[SourceMesh.Indices[ClosestTriangleBaseIdx + 1]];
 			const FVector3f& C = SourceMesh.Positions[SourceMesh.Indices[ClosestTriangleBaseIdx + 2]];
 
-			const FVector3f& NA = SourceMesh.Normals[SourceMesh.Indices[ClosestTriangleBaseIdx]];
-			const FVector3f& NB = SourceMesh.Normals[SourceMesh.Indices[ClosestTriangleBaseIdx + 1]];
-			const FVector3f& NC = SourceMesh.Normals[SourceMesh.Indices[ClosestTriangleBaseIdx + 2]];
+			const FVector3f& NA = SourceMesh.AveragedNormals[SourceMesh.Indices[ClosestTriangleBaseIdx]];
+			const FVector3f& NB = SourceMesh.AveragedNormals[SourceMesh.Indices[ClosestTriangleBaseIdx + 1]];
+			const FVector3f& NC = SourceMesh.AveragedNormals[SourceMesh.Indices[ClosestTriangleBaseIdx + 2]];
 
 			// Before generating the skinning data we need to check for a degenerate triangle.
 			// If we find _any_ degenerate triangles we will notify and fail to generate the skinning data
@@ -447,6 +482,8 @@ namespace ClothingMeshUtils
 										   float KernelMaxDistance,
 										   float MaxIncidentEdgeLength)
 		{
+			check(SourceMesh.AveragedNormals.Num() == SourceMesh.Positions.Num());
+
 			const FVector3f& VertPosition = TargetMesh.Positions[VertIdx0];
 			const FVector3f& VertNormal = TargetMesh.Normals[VertIdx0];
 
@@ -483,9 +520,9 @@ namespace ClothingMeshUtils
 				const FVector3f& B = SourceMesh.Positions[SourceMesh.Indices[ClosestTriangleBaseIdx + 1]];
 				const FVector3f& C = SourceMesh.Positions[SourceMesh.Indices[ClosestTriangleBaseIdx + 2]];
 
-				const FVector3f& NA = SourceMesh.Normals[SourceMesh.Indices[ClosestTriangleBaseIdx]];
-				const FVector3f& NB = SourceMesh.Normals[SourceMesh.Indices[ClosestTriangleBaseIdx + 1]];
-				const FVector3f& NC = SourceMesh.Normals[SourceMesh.Indices[ClosestTriangleBaseIdx + 2]];
+				const FVector3f& NA = SourceMesh.AveragedNormals[SourceMesh.Indices[ClosestTriangleBaseIdx]];
+				const FVector3f& NB = SourceMesh.AveragedNormals[SourceMesh.Indices[ClosestTriangleBaseIdx + 1]];
+				const FVector3f& NC = SourceMesh.AveragedNormals[SourceMesh.Indices[ClosestTriangleBaseIdx + 2]];
 
 				// Before generating the skinning data we need to check for a degenerate triangle.
 				// If we find _any_ degenerate triangles we will notify and fail to generate the skinning data
@@ -616,6 +653,8 @@ namespace ClothingMeshUtils
 		bool bUseMultipleInfluences,
 		float KernelMaxDistance)
 	{
+		check(SourceMesh.AveragedNormals.Num() == SourceMesh.Positions.Num());
+
 		if (!TargetMesh.HasValidMesh())  // Check that the number of positions is equal to the number of normals and that the number of indices is divisible by 3
 		{
 			UE_LOG(LogClothingMeshUtils, Warning, TEXT("Failed to generate mesh to mesh skinning data. Invalid Target Mesh."));
@@ -634,7 +673,7 @@ namespace ClothingMeshUtils
 		const int32 NumMesh0Tris = TargetMesh.Indices.Num() / 3;
 
 		const int32 NumMesh1Verts = SourceMesh.Positions.Num();
-		const int32 NumMesh1Normals = SourceMesh.Normals.Num();
+		const int32 NumMesh1Normals = SourceMesh.AveragedNormals.Num();
 		const int32 NumMesh1Indices = SourceMesh.Indices.Num();
 
 		TargetMesh.ComputeMaxEdgeLengths();
diff --git a/Engine/Source/Runtime/ClothingSystemRuntimeCommon/Public/ClothPhysicalMeshData.h b/Engine/Source/Runtime/ClothingSystemRuntimeCommon/Public/ClothPhysicalMeshData.h
--- a/Engine/Source/Runtime/ClothingSystemRuntimeCommon/Public/ClothPhysicalMeshData.h
+++ b/Engine/Source/Runtime/ClothingSystemRuntimeCommon/Public/ClothPhysicalMeshData.h
@@ -51,12 +51,9 @@ struct CLOTHINGSYSTEMRUNTIMECOMMON_API FClothPhysicalMeshData
 
 	/** Compute vertex normals as unweighted average of incident face normals. If all incident triangles are degenerate,
 	    the vertex normal is assigned a value of FVector3f::XAxisVector. */
-	UE_DEPRECATED(5.0, "Use CalculateNormals() instead.")
+	UE_DEPRECATED(5.0, "This function is no longer part of this API and will soon be removed.")
 	void ComputeFaceAveragedVertexNormals(TArray<FVector3f>& OutNormals) const;
 
-	/** Recalculate the vertex normals as unweighted average of incident face normals. */
-	void CalculateNormals();
-
 	/** Retrieve whether a vertex weight array has already been registered. */
 	template<typename T>
 	bool HasWeightMap(const T Target) const { return WeightMaps.Contains((uint32)Target); }
diff --git a/Engine/Source/Runtime/ClothingSystemRuntimeCommon/Public/Utils/ClothingMeshUtils.h b/Engine/Source/Runtime/ClothingSystemRuntimeCommon/Public/Utils/ClothingMeshUtils.h
--- a/Engine/Source/Runtime/ClothingSystemRuntimeCommon/Public/Utils/ClothingMeshUtils.h
+++ b/Engine/Source/Runtime/ClothingSystemRuntimeCommon/Public/Utils/ClothingMeshUtils.h
@@ -73,6 +73,16 @@ namespace ClothingMeshUtils
 		{
 		}
 
+		/* Construct a mesh descriptor with re-calculated (averaged) normals to match the simulation output. */
+		ClothMeshDesc(TConstArrayView<FVector3f> InPositions, TConstArrayView<uint32> InIndices)
+			: Positions(InPositions)
+			, Indices(InIndices)
+			, bHasValidBVH(false)
+			, bHasValidMaxEdgeLengths(false)
+		{
+			ComputeAveragedNormals();
+		}
+
 		bool HasValidMesh() const
 		{
 			return Positions.Num() == Normals.Num() && Indices.Num() % 3 == 0;
@@ -92,6 +102,13 @@ namespace ClothingMeshUtils
 		 */
 		void ComputeMaxEdgeLengths() const;
 
+		/**
+		 * Compute the face averaged normals.
+		 * For when (cloth deformer), the same normal calculation than the solver outputs are required.
+		 * Declared const for convenience since it doesn't change the descriptor itself.
+		 */
+		void ComputeAveragedNormals() const;
+
 		TConstArrayView<FVector3f> Positions;
 		TConstArrayView<FVector3f> Normals;
 		TConstArrayView<FVector3f> Tangents;
@@ -102,6 +119,7 @@ namespace ClothingMeshUtils
 
 		mutable bool bHasValidMaxEdgeLengths;
 		mutable TArray<float> MaxEdgeLengths;
+		mutable TArray<FVector3f> AveragedNormals;
 	};
 
 	/**
