Cloth - Fix BuildLodTransitionData() to use averaged normals to prevent a crash during cooking of the cloth assets.

diff --git a/Engine/Source/Runtime/ClothingSystemRuntimeCommon/Private/ClothingAsset.cpp b/Engine/Source/Runtime/ClothingSystemRuntimeCommon/Private/ClothingAsset.cpp
--- a/Engine/Source/Runtime/ClothingSystemRuntimeCommon/Private/ClothingAsset.cpp
+++ b/Engine/Source/Runtime/ClothingSystemRuntimeCommon/Private/ClothingAsset.cpp
@@ -866,7 +866,7 @@ void UClothingAssetCommon::BuildLodTransitionData()
 		{
 			FClothPhysicalMeshData& PrevPhysMesh = PrevLod->PhysicalMeshData;
 			CurrentLod.TransitionUpSkinData.Empty(CurrentLodNumVerts);
-			ClothingMeshUtils::ClothMeshDesc PrevMeshDesc(PrevPhysMesh.Vertices, PrevPhysMesh.Normals, PrevPhysMesh.Indices);
+			ClothingMeshUtils::ClothMeshDesc PrevMeshDesc(PrevPhysMesh.Vertices, PrevPhysMesh.Indices);  // Will calculate averaged normals
 
 			ClothingMeshUtils::GenerateMeshToMeshVertData(CurrentLod.TransitionUpSkinData, CurrentMeshDesc, PrevMeshDesc,
 				MaxDistances, bUseSmoothTransitions, bUseMultipleInfluences, SkinningKernelRadius);
@@ -875,7 +875,7 @@ void UClothingAssetCommon::BuildLodTransitionData()
 		{
 			FClothPhysicalMeshData& NextPhysMesh = NextLod->PhysicalMeshData;
 			CurrentLod.TransitionDownSkinData.Empty(CurrentLodNumVerts);
-			ClothingMeshUtils::ClothMeshDesc NextMeshDesc(NextPhysMesh.Vertices, NextPhysMesh.Normals, NextPhysMesh.Indices);
+			ClothingMeshUtils::ClothMeshDesc NextMeshDesc(NextPhysMesh.Vertices, NextPhysMesh.Indices);  // Will calculate averaged normals 
 			ClothingMeshUtils::GenerateMeshToMeshVertData(CurrentLod.TransitionDownSkinData, CurrentMeshDesc, NextMeshDesc,
 				MaxDistances, bUseSmoothTransitions, bUseMultipleInfluences, SkinningKernelRadius);
 		}
