Cloth - Fix the mesh descriptor averaged normals initializations that was preventing cloth binding to work.

diff --git a/Engine/Source/Runtime/ClothingSystemRuntimeCommon/Public/Utils/ClothingMeshUtils.h b/Engine/Source/Runtime/ClothingSystemRuntimeCommon/Public/Utils/ClothingMeshUtils.h
--- a/Engine/Source/Runtime/ClothingSystemRuntimeCommon/Public/Utils/ClothingMeshUtils.h
+++ b/Engine/Source/Runtime/ClothingSystemRuntimeCommon/Public/Utils/ClothingMeshUtils.h
@@ -81,6 +81,7 @@ namespace ClothingMeshUtils
 			, bHasValidMaxEdgeLengths(false)
 		{
 			ComputeAveragedNormals();
+			Normals = AveragedNormals;
 		}
 
 		bool HasValidMesh() const
