Cloth - Tidy up the ClothingMeshUtils files.

diff --git a/Engine/Source/Runtime/ClothingSystemRuntimeCommon/Private/Utils/ClothingMeshUtils.cpp b/Engine/Source/Runtime/ClothingSystemRuntimeCommon/Private/Utils/ClothingMeshUtils.cpp
--- a/Engine/Source/Runtime/ClothingSystemRuntimeCommon/Private/Utils/ClothingMeshUtils.cpp
+++ b/Engine/Source/Runtime/ClothingSystemRuntimeCommon/Private/Utils/ClothingMeshUtils.cpp
@@ -23,42 +23,125 @@ DEFINE_LOG_CATEGORY(LogClothingMeshUtils)
 
 namespace ClothingMeshUtils
 {
-	void ClothMeshDesc::ComputeMaxEdgeLengths() const
+	struct ClothMeshDesc::FClothBvEntry
 	{
-		if (bHasValidMaxEdgeLengths)
+		const ClothMeshDesc* TmData;
+		int32 Index;
+
+		bool HasBoundingBox() const { return true; }
+
+		Chaos::FAABB3 BoundingBox() const
 		{
-			return;
+			int32 TriBaseIdx = Index * 3;
+
+			const uint32 IA = TmData->Indices[TriBaseIdx + 0];
+			const uint32 IB = TmData->Indices[TriBaseIdx + 1];
+			const uint32 IC = TmData->Indices[TriBaseIdx + 2];
+
+			const FVector3f& A = TmData->Positions[IA];
+			const FVector3f& B = TmData->Positions[IB];
+			const FVector3f& C = TmData->Positions[IC];
+
+			Chaos::FAABB3 Bounds(A, A);
+
+			Bounds.GrowToInclude(B);
+			Bounds.GrowToInclude(C);
+
+			return Bounds;
 		}
-		ensure(Indices.Num() % 3 == 0);  // Check we have properly formed triangles
 
-		const int32 NumMesh0Verts = Positions.Num();
-		const int32 NumMesh0Tris = Indices.Num() / 3;
+		template<typename TPayloadType>
+		int32 GetPayload(int32 Idx) const { return Idx; }
+	};
 
-		MaxEdgeLengths.Init(0.f, NumMesh0Verts);
+	float ClothMeshDesc::DistanceToTriangle(const FVector& Position, int32 TriangleBaseIndex) const
+	{
+		const uint32 IA = Indices[TriangleBaseIndex + 0];
+		const uint32 IB = Indices[TriangleBaseIndex + 1];
+		const uint32 IC = Indices[TriangleBaseIndex + 2];
 
-		for (int32 TriangleIdx = 0; TriangleIdx < NumMesh0Tris; ++TriangleIdx)
-		{
-			const uint32* const Triangle = &Indices[TriangleIdx * 3];
+		const FVector& A = Positions[IA];
+		const FVector& B = Positions[IB];
+		const FVector& C = Positions[IC];
 
-			for (int32 Vertex0Idx = 0; Vertex0Idx < 3; ++Vertex0Idx)
+		const FVector PointOnTri = FMath::ClosestPointOnTriangleToPoint(Position, A, B, C);
+		return (PointOnTri - Position).Size();
+	}
+
+	TArray<int32> ClothMeshDesc::FindCandidateTriangles(const FVector& InPoint, float InTolerance) const
+	{
+		ensure(HasValidMesh());
+		static const int32 MinNumTrianglesForBVHCreation = 100;
+		const int32 NumTris = Indices.Num() / 3;
+		if (NumTris > MinNumTrianglesForBVHCreation)
+		{
+			// This is not thread safe
+			if (!bHasValidBVH)
 			{
-				const int32 Vertex1Idx = (Vertex0Idx + 1) % 3;
+				TArray<FClothBvEntry> BVEntries;
+				BVEntries.Reset(NumTris);
+
+				for (int32 Tri = 0; Tri < NumTris; ++Tri)
+				{
+					BVEntries.Add({ this, Tri });
+				}
+				BVH.Reinitialize(BVEntries);
+				bHasValidBVH = true;
+			}
+			Chaos::FAABB3 TmpAABB(InPoint, InPoint);
+			TmpAABB.Thicken(InTolerance);  // Most points might be very close to the triangle, but not directly on it
+			TArray<int32> Triangles = BVH.FindAllIntersections(TmpAABB);
 
-				const FVector& P0 = Positions[Triangle[Vertex0Idx]];
-				const FVector& P1 = Positions[Triangle[Vertex1Idx]];
+			// Refine the search to include all nearby bounded volumes (the point could well be outside the closest triangle's bounded volume)
+			if (Triangles.Num())
+			{
+				float ClosestDistance = TNumericLimits<float>::Max();
+				for (const int32 Triangle : Triangles)
+				{
+					ClosestDistance = FMath::Min(ClosestDistance, DistanceToTriangle(InPoint, Triangle * 3));
+				}
 
-				const float EdgeLength = FVector::Distance(P0, P1);
-				MaxEdgeLengths[Triangle[Vertex0Idx]] = FMath::Max(MaxEdgeLengths[Triangle[Vertex0Idx]], EdgeLength);
-				MaxEdgeLengths[Triangle[Vertex1Idx]] = FMath::Max(MaxEdgeLengths[Triangle[Vertex1Idx]], EdgeLength);
+				TmpAABB.Thicken(ClosestDistance);
+				return BVH.FindAllIntersections(TmpAABB);
 			}
 		}
+		return TArray<int32>();
+	}
 
-		bHasValidMaxEdgeLengths = true;
+	TConstArrayView<float> ClothMeshDesc::GetMaxEdgeLengths() const
+	{
+		if (!MaxEdgeLengths.Num())
+		{
+			ensure(Indices.Num() % 3 == 0);  // Check we have properly formed triangles
+
+			const int32 NumMesh0Verts = Positions.Num();
+			const int32 NumMesh0Tris = Indices.Num() / 3;
+
+			MaxEdgeLengths.Init(0.f, NumMesh0Verts);
+
+			for (int32 TriangleIdx = 0; TriangleIdx < NumMesh0Tris; ++TriangleIdx)
+			{
+				const uint32* const Triangle = &Indices[TriangleIdx * 3];
+
+				for (int32 Vertex0Idx = 0; Vertex0Idx < 3; ++Vertex0Idx)
+				{
+					const int32 Vertex1Idx = (Vertex0Idx + 1) % 3;
+
+					const FVector& P0 = Positions[Triangle[Vertex0Idx]];
+					const FVector& P1 = Positions[Triangle[Vertex1Idx]];
+
+					const float EdgeLength = FVector::Distance(P0, P1);
+					MaxEdgeLengths[Triangle[Vertex0Idx]] = FMath::Max(MaxEdgeLengths[Triangle[Vertex0Idx]], EdgeLength);
+					MaxEdgeLengths[Triangle[Vertex1Idx]] = FMath::Max(MaxEdgeLengths[Triangle[Vertex1Idx]], EdgeLength);
+				}
+			}
+		}
+		return TConstArrayView<float>(MaxEdgeLengths);
 	}
 
-	void ClothMeshDesc::ComputeAveragedNormals() const
+	void ClothMeshDesc::ComputeAveragedNormals()
 	{
-		AveragedNormals.Init(FVector3f::ZeroVector, Positions.Num());  // Note AveragedNormals is a mutable, so can be modified in this function
+		AveragedNormals.Init(FVector3f::ZeroVector, Positions.Num());
 
 		const int32 NumTriangles = Indices.Num() / 3;
 		for (int32 TriangleIndex = 0; TriangleIndex < NumTriangles; ++TriangleIndex)
@@ -278,23 +361,23 @@ namespace ClothingMeshUtils
 		float MinimumDistanceSq = MAX_flt;
 		int32 ClosestBaseIndex = INDEX_NONE;
 
-		const TArray<int32> Tris = const_cast<ClothMeshDesc&>(Mesh).FindCandidateTriangles(Position, InTolerance);
+		const TArray<int32> Tris = Mesh.FindCandidateTriangles(Position, InTolerance);
 		int32 NumTriangles = Tris.Num();
 		if (!NumTriangles)
 		{
-			NumTriangles = Mesh.Indices.Num() / 3;
+			NumTriangles = Mesh.GetIndices().Num() / 3;
 		}
 		for (int32 TriIdx = 0; TriIdx < NumTriangles; ++TriIdx)
 		{
 			int32 TriBaseIdx = (Tris.Num() ? Tris[TriIdx] : TriIdx) * 3;
 
-			const uint32 IA = Mesh.Indices[TriBaseIdx + 0];
-			const uint32 IB = Mesh.Indices[TriBaseIdx + 1];
-			const uint32 IC = Mesh.Indices[TriBaseIdx + 2];
+			const uint32 IA = Mesh.GetIndices()[TriBaseIdx + 0];
+			const uint32 IB = Mesh.GetIndices()[TriBaseIdx + 1];
+			const uint32 IC = Mesh.GetIndices()[TriBaseIdx + 2];
 
-			const FVector& A = Mesh.Positions[IA];
-			const FVector& B = Mesh.Positions[IB];
-			const FVector& C = Mesh.Positions[IC];
+			const FVector& A = Mesh.GetPositions()[IA];
+			const FVector& B = Mesh.GetPositions()[IB];
+			const FVector& C = Mesh.GetPositions()[IC];
 
 			FVector PointOnTri = FMath::ClosestPointOnTriangleToPoint(Position, A, B, C);
 			float DistSq = (PointOnTri - Position).SizeSquared();
@@ -311,35 +394,20 @@ namespace ClothingMeshUtils
 
 	namespace  // Helpers
 	{
-
-		float DistanceToTriangle(const FVector& Position, const ClothMeshDesc& Mesh, int32 TriBaseIdx)
-		{
-			const uint32 IA = Mesh.Indices[TriBaseIdx + 0];
-			const uint32 IB = Mesh.Indices[TriBaseIdx + 1];
-			const uint32 IC = Mesh.Indices[TriBaseIdx + 2];
-
-			const FVector& A = Mesh.Positions[IA];
-			const FVector& B = Mesh.Positions[IB];
-			const FVector& C = Mesh.Positions[IC];
-
-			FVector PointOnTri = FMath::ClosestPointOnTriangleToPoint(Position, A, B, C);
-			return (PointOnTri - Position).Size();
-		}
-
 		using TriangleDistance = TPair<int32, float>;
 
 		/** Similar to GetBestTriangleBaseIndex but returns the N closest triangles. */
 		template<uint32 N>
 		TStaticArray<TriangleDistance, N> GetNBestTrianglesBaseIndices(const ClothMeshDesc& Mesh, const FVector& Position)
 		{
-			const TArray<int32> Tris = const_cast<ClothMeshDesc&>(Mesh).FindCandidateTriangles(Position);
+			const TArray<int32> Tris = Mesh.FindCandidateTriangles(Position);
 			int32 NumTriangles = Tris.Num();
 
 			if (NumTriangles < N)
 			{
 				// Couldn't find N candidates using FindCandidateTriangles. Grab all triangles in the mesh and see if
 				// we have enough
-				NumTriangles = Mesh.Indices.Num() / 3;
+				NumTriangles = Mesh.GetIndices().Num() / 3;
 
 				if (NumTriangles < N)
 				{
@@ -350,7 +418,7 @@ namespace ClothingMeshUtils
 					for ( ; i < NumTriangles; ++i)
 					{
 						int32 TriBaseIdx = 3*i;
-						float CurrentDistance = DistanceToTriangle(Position, Mesh, TriBaseIdx);
+						float CurrentDistance = Mesh.DistanceToTriangle(Position, TriBaseIdx);
 						ClosestTriangles[i] = TPair<int32, float>{ TriBaseIdx, CurrentDistance };
 					}
 
@@ -374,7 +442,7 @@ namespace ClothingMeshUtils
 			for (; i < N; ++i)
 			{
 				int32 TriBaseIdx = (Tris.Num() >= N ? Tris[i] : i) * 3;
-				float CurrentDistance = DistanceToTriangle(Position, Mesh, TriBaseIdx);
+				float CurrentDistance = Mesh.DistanceToTriangle(Position, TriBaseIdx);
 				ClosestTriangles[i] = TPair<int32, float>{ TriBaseIdx, CurrentDistance };
 			}
 
@@ -388,7 +456,7 @@ namespace ClothingMeshUtils
 			for (; i < NumTriangles; ++i)
 			{
 				int32 TriBaseIdx = (Tris.Num() >= N ? Tris[i] : i) * 3;
-				float CurrentDistance = DistanceToTriangle(Position, Mesh, TriBaseIdx);
+				float CurrentDistance = Mesh.DistanceToTriangle(Position, TriBaseIdx);
 
 				if (CurrentDistance < ClosestTriangles[0].Value)
 				{
@@ -431,15 +499,15 @@ namespace ClothingMeshUtils
 										 int32 ClosestTriangleBaseIdx,
 										 FMeshToMeshVertData& OutSkinningData)
 		{
-			check(SourceMesh.AveragedNormals.Num() == SourceMesh.Positions.Num());
+			check(SourceMesh.HasAveragedNormals());
 
-			const FVector3f& A = SourceMesh.Positions[SourceMesh.Indices[ClosestTriangleBaseIdx]];
-			const FVector3f& B = SourceMesh.Positions[SourceMesh.Indices[ClosestTriangleBaseIdx + 1]];
-			const FVector3f& C = SourceMesh.Positions[SourceMesh.Indices[ClosestTriangleBaseIdx + 2]];
+			const FVector3f& A = SourceMesh.GetPositions()[SourceMesh.GetIndices()[ClosestTriangleBaseIdx]];
+			const FVector3f& B = SourceMesh.GetPositions()[SourceMesh.GetIndices()[ClosestTriangleBaseIdx + 1]];
+			const FVector3f& C = SourceMesh.GetPositions()[SourceMesh.GetIndices()[ClosestTriangleBaseIdx + 2]];
 
-			const FVector3f& NA = SourceMesh.AveragedNormals[SourceMesh.Indices[ClosestTriangleBaseIdx]];
-			const FVector3f& NB = SourceMesh.AveragedNormals[SourceMesh.Indices[ClosestTriangleBaseIdx + 1]];
-			const FVector3f& NC = SourceMesh.AveragedNormals[SourceMesh.Indices[ClosestTriangleBaseIdx + 2]];
+			const FVector3f& NA = SourceMesh.GetNormals()[SourceMesh.GetIndices()[ClosestTriangleBaseIdx]];
+			const FVector3f& NB = SourceMesh.GetNormals()[SourceMesh.GetIndices()[ClosestTriangleBaseIdx + 1]];
+			const FVector3f& NC = SourceMesh.GetNormals()[SourceMesh.GetIndices()[ClosestTriangleBaseIdx + 2]];
 
 			// Before generating the skinning data we need to check for a degenerate triangle.
 			// If we find _any_ degenerate triangles we will notify and fail to generate the skinning data
@@ -464,9 +532,9 @@ namespace ClothingMeshUtils
 			OutSkinningData.PositionBaryCoordsAndDist = GetPointBaryAndDistWithNormals(A, B, C, NA, NB, NC, VertPosition);
 			OutSkinningData.NormalBaryCoordsAndDist = GetPointBaryAndDistWithNormals(A, B, C, NA, NB, NC, VertPosition + VertNormal);
 			OutSkinningData.TangentBaryCoordsAndDist = GetPointBaryAndDistWithNormals(A, B, C, NA, NB, NC, VertPosition + VertTangent);
-			OutSkinningData.SourceMeshVertIndices[0] = SourceMesh.Indices[ClosestTriangleBaseIdx];
-			OutSkinningData.SourceMeshVertIndices[1] = SourceMesh.Indices[ClosestTriangleBaseIdx + 1];
-			OutSkinningData.SourceMeshVertIndices[2] = SourceMesh.Indices[ClosestTriangleBaseIdx + 2];
+			OutSkinningData.SourceMeshVertIndices[0] = SourceMesh.GetIndices()[ClosestTriangleBaseIdx];
+			OutSkinningData.SourceMeshVertIndices[1] = SourceMesh.GetIndices()[ClosestTriangleBaseIdx + 1];
+			OutSkinningData.SourceMeshVertIndices[2] = SourceMesh.GetIndices()[ClosestTriangleBaseIdx + 2];
 			OutSkinningData.SourceMeshVertIndices[3] = 0;
 			OutSkinningData.Weight = 1.0f;
 
@@ -482,15 +550,15 @@ namespace ClothingMeshUtils
 										   float KernelMaxDistance,
 										   float MaxIncidentEdgeLength)
 		{
-			check(SourceMesh.AveragedNormals.Num() == SourceMesh.Positions.Num());
+			check(SourceMesh.HasAveragedNormals());
 
-			const FVector3f& VertPosition = TargetMesh.Positions[VertIdx0];
-			const FVector3f& VertNormal = TargetMesh.Normals[VertIdx0];
+			const FVector3f& VertPosition = TargetMesh.GetPositions()[VertIdx0];
+			const FVector3f& VertNormal = TargetMesh.GetNormals()[VertIdx0];
 
 			FVector VertTangent;
 			if (TargetMesh.HasTangents())
 			{
-				VertTangent = TargetMesh.Tangents[VertIdx0];
+				VertTangent = TargetMesh.GetTangents()[VertIdx0];
 			}
 			else
 			{
@@ -516,13 +584,13 @@ namespace ClothingMeshUtils
 					continue;
 				}
 
-				const FVector3f& A = SourceMesh.Positions[SourceMesh.Indices[ClosestTriangleBaseIdx]];
-				const FVector3f& B = SourceMesh.Positions[SourceMesh.Indices[ClosestTriangleBaseIdx + 1]];
-				const FVector3f& C = SourceMesh.Positions[SourceMesh.Indices[ClosestTriangleBaseIdx + 2]];
+				const FVector3f& A = SourceMesh.GetPositions()[SourceMesh.GetIndices()[ClosestTriangleBaseIdx]];
+				const FVector3f& B = SourceMesh.GetPositions()[SourceMesh.GetIndices()[ClosestTriangleBaseIdx + 1]];
+				const FVector3f& C = SourceMesh.GetPositions()[SourceMesh.GetIndices()[ClosestTriangleBaseIdx + 2]];
 
-				const FVector3f& NA = SourceMesh.AveragedNormals[SourceMesh.Indices[ClosestTriangleBaseIdx]];
-				const FVector3f& NB = SourceMesh.AveragedNormals[SourceMesh.Indices[ClosestTriangleBaseIdx + 1]];
-				const FVector3f& NC = SourceMesh.AveragedNormals[SourceMesh.Indices[ClosestTriangleBaseIdx + 2]];
+				const FVector3f& NA = SourceMesh.GetNormals()[SourceMesh.GetIndices()[ClosestTriangleBaseIdx]];
+				const FVector3f& NB = SourceMesh.GetNormals()[SourceMesh.GetIndices()[ClosestTriangleBaseIdx + 1]];
+				const FVector3f& NC = SourceMesh.GetNormals()[SourceMesh.GetIndices()[ClosestTriangleBaseIdx + 2]];
 
 				// Before generating the skinning data we need to check for a degenerate triangle.
 				// If we find _any_ degenerate triangles we will notify and fail to generate the skinning data
@@ -547,9 +615,9 @@ namespace ClothingMeshUtils
 				CurrentData.PositionBaryCoordsAndDist = GetPointBaryAndDistWithNormals(A, B, C, NA, NB, NC, VertPosition);
 				CurrentData.NormalBaryCoordsAndDist = GetPointBaryAndDistWithNormals(A, B, C, NA, NB, NC, VertPosition + VertNormal);
 				CurrentData.TangentBaryCoordsAndDist = GetPointBaryAndDistWithNormals(A, B, C, NA, NB, NC, VertPosition + VertTangent);
-				CurrentData.SourceMeshVertIndices[0] = SourceMesh.Indices[ClosestTriangleBaseIdx];
-				CurrentData.SourceMeshVertIndices[1] = SourceMesh.Indices[ClosestTriangleBaseIdx + 1];
-				CurrentData.SourceMeshVertIndices[2] = SourceMesh.Indices[ClosestTriangleBaseIdx + 2];
+				CurrentData.SourceMeshVertIndices[0] = SourceMesh.GetIndices()[ClosestTriangleBaseIdx];
+				CurrentData.SourceMeshVertIndices[1] = SourceMesh.GetIndices()[ClosestTriangleBaseIdx + 1];
+				CurrentData.SourceMeshVertIndices[2] = SourceMesh.GetIndices()[ClosestTriangleBaseIdx + 2];
 				CurrentData.SourceMeshVertIndices[3] = 0;
 
 				CurrentData.Weight = Kernel(NearestTriangles[j].Value, KernelMaxDistance);
@@ -575,15 +643,14 @@ namespace ClothingMeshUtils
 				return;
 			}
 
-			const int32 NumTargetMeshVerts = TargetMesh.Positions.Num();
+			const int32 NumTargetMeshVerts = TargetMesh.GetPositions().Num();
 
 			if (!ensure(NumTargetMeshVerts * NUM_INFLUENCES_PER_VERTEX == InOutSkinningData.Num()))
 			{
 				return;
 			}
 
-			TargetMesh.ComputeMaxEdgeLengths();
-			const TArray<float>& MaxEdgeLengths = TargetMesh.MaxEdgeLengths;
+			const TConstArrayView<float> MaxEdgeLengths = TargetMesh.GetMaxEdgeLengths();
 			check(NumTargetMeshVerts == MaxEdgeLengths.Num());
 
 			for (int32 VID = 0; VID < NumTargetMeshVerts; ++VID)
@@ -601,13 +668,13 @@ namespace ClothingMeshUtils
 				if (SumWeight < KINDA_SMALL_NUMBER)
 				{
 					// Fall back to single influence
-					const FVector3f& VertPosition = TargetMesh.Positions[VID];
-					const FVector3f& VertNormal = TargetMesh.Normals[VID];
+					const FVector3f& VertPosition = TargetMesh.GetPositions()[VID];
+					const FVector3f& VertNormal = TargetMesh.GetNormals()[VID];
 
 					FVector3f VertTangent;
 					if (TargetMesh.HasTangents())
 					{
-						VertTangent = TargetMesh.Tangents[VID];
+						VertTangent = TargetMesh.GetTangents()[VID];
 					}
 					else
 					{
@@ -653,7 +720,7 @@ namespace ClothingMeshUtils
 		bool bUseMultipleInfluences,
 		float KernelMaxDistance)
 	{
-		check(SourceMesh.AveragedNormals.Num() == SourceMesh.Positions.Num());
+		check(SourceMesh.HasAveragedNormals());
 
 		if (!TargetMesh.HasValidMesh())  // Check that the number of positions is equal to the number of normals and that the number of indices is divisible by 3
 		{
@@ -667,17 +734,16 @@ namespace ClothingMeshUtils
 			return;
 		}
 
-		const int32 NumMesh0Verts = TargetMesh.Positions.Num();
-		const int32 NumMesh0Normals = TargetMesh.Normals.Num();
-		const int32 NumMesh0Tangents = TargetMesh.HasTangents() ? TargetMesh.Tangents.Num() : 0;
-		const int32 NumMesh0Tris = TargetMesh.Indices.Num() / 3;
+		const int32 NumMesh0Verts = TargetMesh.GetPositions().Num();
+		const int32 NumMesh0Normals = TargetMesh.GetNormals().Num();
+		const int32 NumMesh0Tangents = TargetMesh.HasTangents() ? TargetMesh.GetTangents().Num() : 0;
+		const int32 NumMesh0Tris = TargetMesh.GetIndices().Num() / 3;
 
-		const int32 NumMesh1Verts = SourceMesh.Positions.Num();
-		const int32 NumMesh1Normals = SourceMesh.AveragedNormals.Num();
-		const int32 NumMesh1Indices = SourceMesh.Indices.Num();
+		const int32 NumMesh1Verts = SourceMesh.GetPositions().Num();
+		const int32 NumMesh1Normals = SourceMesh.GetNormals().Num();
+		const int32 NumMesh1Indices = SourceMesh.GetIndices().Num();
 
-		TargetMesh.ComputeMaxEdgeLengths();
-		const TArray<float>& MaxEdgeLengths = TargetMesh.MaxEdgeLengths;
+		const TConstArrayView<float> MaxEdgeLengths = TargetMesh.GetMaxEdgeLengths();
 		check(NumMesh0Verts == MaxEdgeLengths.Num());
 
 		if (bUseMultipleInfluences)
@@ -719,13 +785,13 @@ namespace ClothingMeshUtils
 				OutMeshToMeshVertData.AddZeroed();
 				FMeshToMeshVertData& SkinningData = OutMeshToMeshVertData.Last();
 
-				const FVector3f& VertPosition = TargetMesh.Positions[VertIdx0];
-				const FVector3f& VertNormal = TargetMesh.Normals[VertIdx0];
+				const FVector3f& VertPosition = TargetMesh.GetPositions()[VertIdx0];
+				const FVector3f& VertNormal = TargetMesh.GetNormals()[VertIdx0];
 
 				FVector VertTangent;
 				if (TargetMesh.HasTangents())
 				{
-					VertTangent = TargetMesh.Tangents[VertIdx0];
+					VertTangent = TargetMesh.GetTangents()[VertIdx0];
 				}
 				else
 				{
@@ -769,7 +835,7 @@ namespace ClothingMeshUtils
 			return;
 		}
 
-		const int32 NumTargetMeshVerts = TargetMesh.Positions.Num();
+		const int32 NumTargetMeshVerts = TargetMesh.GetPositions().Num();
 
 		if (!ensure(NumTargetMeshVerts * NUM_INFLUENCES_PER_VERTEX == InOutSkinningData.Num()))
 		{
@@ -796,8 +862,8 @@ namespace ClothingMeshUtils
 			if (SumWeight < KINDA_SMALL_NUMBER)
 			{
 				// Fall back to single influence
-				const FVector3f& VertPosition = TargetMesh.Positions[VID];
-				const FVector3f& VertNormal = TargetMesh.Normals[VID];
+				const FVector3f& VertPosition = TargetMesh.GetPositions()[VID];
+				const FVector3f& VertNormal = TargetMesh.GetNormals()[VID];
 
 				FVector3f VertTangent;
 				if (TargetTangents)
@@ -1176,17 +1242,17 @@ namespace ClothingMeshUtils
 
 			int32 TriBaseIndex = GetBestTriangleBaseIndex(SourceMesh, Position);
 
-			const int32 IA = SourceMesh.Indices[TriBaseIndex];
-			const int32 IB = SourceMesh.Indices[TriBaseIndex + 1];
-			const int32 IC = SourceMesh.Indices[TriBaseIndex + 2];
+			const int32 IA = SourceMesh.GetIndices()[TriBaseIndex];
+			const int32 IB = SourceMesh.GetIndices()[TriBaseIndex + 1];
+			const int32 IC = SourceMesh.GetIndices()[TriBaseIndex + 2];
 
-			const FVector3f& A = SourceMesh.Positions[IA];
-			const FVector3f& B = SourceMesh.Positions[IB];
-			const FVector3f& C = SourceMesh.Positions[IC];
+			const FVector3f& A = SourceMesh.GetPositions()[IA];
+			const FVector3f& B = SourceMesh.GetPositions()[IB];
+			const FVector3f& C = SourceMesh.GetPositions()[IC];
 
-			const FVector3f& NA = SourceMesh.Normals[IA];
-			const FVector3f& NB = SourceMesh.Normals[IB];
-			const FVector3f& NC = SourceMesh.Normals[IC];
+			const FVector3f& NA = SourceMesh.GetNormals()[IA];
+			const FVector3f& NB = SourceMesh.GetNormals()[IB];
+			const FVector3f& NC = SourceMesh.GetNormals()[IC];
 
 			OutEmbeddedPositions[PositionIndex] = GetPointBaryAndDistWithNormals(A, B, C, NA, NB, NC, Position);
 			OutSourceIndices.Add(IA);
@@ -1291,47 +1357,6 @@ namespace ClothingMeshUtils
 			OutSourceIndices.Add(IC);
 		}
 	}
-
-	TArray<int32> ClothMeshDesc::FindCandidateTriangles(const FVector& InPoint, float InTolerance)
-	{
-		ensure(HasValidMesh());
-		static const int32 MinNumTrianglesForBVHCreation = 100;
-		const int32 NumTris = Indices.Num() / 3;
-		if (NumTris > MinNumTrianglesForBVHCreation)
-		{
-			// This is not thread safe
-			if (!bHasValidBVH)
-			{
-				TArray<FClothBvEntry> BVEntries;
-				BVEntries.Reset(NumTris);
-
-				for (int32 Tri = 0; Tri < NumTris; ++Tri)
-				{
-					BVEntries.Add({ this, Tri });
-				}
-				BVH.Reinitialize(BVEntries);
-				bHasValidBVH = true;
-			}
-			Chaos::FAABB3 TmpAABB(InPoint, InPoint);
-			TmpAABB.Thicken(InTolerance);  // Most points might be very close to the triangle, but not directly on it
-			TArray<int32> Triangles = BVH.FindAllIntersections(TmpAABB);
-
-			// Refine the search to include all nearby bounded volumes (the point could well be outside the closest triangle's bounded volume)
-			if (Triangles.Num())
-			{
-				float ClosestDistance = TNumericLimits<float>::Max();
-				for (const int32 Triangle : Triangles)
-				{
-					ClosestDistance = FMath::Min(ClosestDistance, DistanceToTriangle(InPoint, *this, Triangle * 3));
-				}
-
-				TmpAABB.Thicken(ClosestDistance);
-				return BVH.FindAllIntersections(TmpAABB);
-			}
-		}
-		return TArray<int32>();
-	}
-
 }
 
 #undef LOCTEXT_NAMESPACE
diff --git a/Engine/Source/Runtime/ClothingSystemRuntimeCommon/Public/Utils/ClothingMeshUtils.h b/Engine/Source/Runtime/ClothingSystemRuntimeCommon/Public/Utils/ClothingMeshUtils.h
--- a/Engine/Source/Runtime/ClothingSystemRuntimeCommon/Public/Utils/ClothingMeshUtils.h
+++ b/Engine/Source/Runtime/ClothingSystemRuntimeCommon/Public/Utils/ClothingMeshUtils.h
@@ -4,7 +4,6 @@
 
 #include "SkeletalMeshTypes.h"
 #include "Chaos/AABBTree.h"
-#include "Chaos/GeometryParticlesfwd.h"
 
 DECLARE_LOG_CATEGORY_EXTERN(LogClothingMeshUtils, Log, All);
 
@@ -13,53 +12,14 @@ struct FPointWeightMap;
 
 namespace ClothingMeshUtils
 {
-	struct CLOTHINGSYSTEMRUNTIMECOMMON_API ClothMeshDesc
+	class CLOTHINGSYSTEMRUNTIMECOMMON_API ClothMeshDesc
 	{
-		struct FClothBvEntry
-		{
-			ClothMeshDesc* TmData;
-			int32 Index;
-
-			bool HasBoundingBox() const { return true; }
-
-			Chaos::FAABB3 BoundingBox() const
-			{
-				int32 TriBaseIdx = Index * 3;
-
-				const uint32 IA = TmData->Indices[TriBaseIdx + 0];
-				const uint32 IB = TmData->Indices[TriBaseIdx + 1];
-				const uint32 IC = TmData->Indices[TriBaseIdx + 2];
-
-				const FVector3f& A = TmData->Positions[IA];
-				const FVector3f& B = TmData->Positions[IB];
-				const FVector3f& C = TmData->Positions[IC];
-
-				Chaos::FAABB3 Bounds(A, A);
-
-				Bounds.GrowToInclude(B);
-				Bounds.GrowToInclude(C);
-
-				return Bounds;
-			}
-
-			template<typename TPayloadType>
-			int32 GetPayload(int32 Idx) const
-			{
-				return Idx;
-			}
-
-			Chaos::FUniqueIdx UniqueIdx() const
-			{
-				return Chaos::FUniqueIdx(Index);
-			}
-		};
-
+	public:
 		ClothMeshDesc(TConstArrayView<FVector3f> InPositions, TConstArrayView<FVector3f> InNormals, TConstArrayView<uint32> InIndices)
 			: Positions(InPositions)
 			, Normals(InNormals)
 			, Indices(InIndices)
 			, bHasValidBVH(false)
-			, bHasValidMaxEdgeLengths(false)
 		{
 		}
 
@@ -69,7 +29,6 @@ namespace ClothingMeshUtils
 			, Tangents(InTangents)
 			, Indices(InIndices)
 			, bHasValidBVH(false)
-			, bHasValidMaxEdgeLengths(false)
 		{
 		}
 
@@ -78,49 +37,53 @@ namespace ClothingMeshUtils
 			: Positions(InPositions)
 			, Indices(InIndices)
 			, bHasValidBVH(false)
-			, bHasValidMaxEdgeLengths(false)
 		{
 			ComputeAveragedNormals();
-			Normals = AveragedNormals;
+			Normals = TConstArrayView<FVector3f>(AveragedNormals);
 		}
 
-		bool HasValidMesh() const
-		{
-			return Positions.Num() == Normals.Num() && Indices.Num() % 3 == 0;
-		}
+		const TConstArrayView<FVector3f>& GetPositions() const { return Positions; }
+		const TConstArrayView<FVector3f>& GetNormals() const { return Normals; }
+		const TConstArrayView<FVector3f>& GetTangents() const { return Tangents; }
+		const TConstArrayView<uint32>& GetIndices() const { return Indices; }
 
-		bool HasTangents() const
-		{
-			return Positions.Num() == Tangents.Num();
-		}
+		/** Return whether the mesh descriptor has a valid number of normals and triangle indices. */
+		bool HasValidMesh() const { return Positions.Num() == Normals.Num() && Indices.Num() % 3 == 0; }
+
+		/** Return whether this descriptor has been initialized with tangent information. */
+		bool HasTangents() const { return Positions.Num() == Tangents.Num(); }
 
-		TArray<int32> FindCandidateTriangles(const FVector& InPoint, float InTolerance = KINDA_SMALL_NUMBER);
+		/** Return whether this descriptor has been initialized with the normals generated from the averaged triangle normals. */
+		bool HasAveragedNormals() const { return Normals.Num() == AveragedNormals.Num(); }
+
+		/** Find the distance to the specified triangle from a point. */
+		float DistanceToTriangle(const FVector& Position, int32 TriangleBaseIndex) const;
+
+		/** Find the closest triangles from the specified point. */
+		TArray<int32> FindCandidateTriangles(const FVector& InPoint, float InTolerance = KINDA_SMALL_NUMBER) const;
 
 		/**
-		 * Compute the max edge length for each edge coming out of a mesh vertex.
+		 * Return the max edge length for each edge coming out of a mesh vertex.
 		 * Useful for guiding the search radius when searching for nearest triangles.
-		 * Declared const for convenience since it doesn't change the descriptor itself.
 		 */
-		void ComputeMaxEdgeLengths() const;
+		TConstArrayView<float> GetMaxEdgeLengths() const;
 
-		/**
-		 * Compute the face averaged normals.
-		 * For when (cloth deformer), the same normal calculation than the solver outputs are required.
-		 * Declared const for convenience since it doesn't change the descriptor itself.
-		 */
-		void ComputeAveragedNormals() const;
+	private:
+		void ComputeAveragedNormals();
+
+		struct FClothBvEntry;
 
 		TConstArrayView<FVector3f> Positions;
 		TConstArrayView<FVector3f> Normals;
 		TConstArrayView<FVector3f> Tangents;
 		TConstArrayView<uint32> Indices;
 
-		bool bHasValidBVH;
-		Chaos::TAABBTree<int32, Chaos::TAABBTreeLeafArray<int32, false>, false> BVH;
+		TArray<FVector3f> AveragedNormals;
+
+		mutable bool bHasValidBVH;
+		mutable Chaos::TAABBTree<int32, Chaos::TAABBTreeLeafArray<int32, false>, false> BVH;
 
-		mutable bool bHasValidMaxEdgeLengths;
 		mutable TArray<float> MaxEdgeLengths;
-		mutable TArray<FVector3f> AveragedNormals;
 	};
 
 	/**
@@ -180,8 +143,7 @@ namespace ClothingMeshUtils
 	UE_DEPRECATED(5.0, "Superseded by GenerateMeshToMeshVertData.")
 	inline void ComputeMaxEdgeLength(const ClothMeshDesc& TargetMesh, TArray<float>& OutMaxEdgeLength)
 	{
-		TargetMesh.ComputeMaxEdgeLengths();
-		OutMaxEdgeLength = TargetMesh.MaxEdgeLengths;
+		OutMaxEdgeLength = TargetMesh.GetMaxEdgeLengths();
 	}
 
 	/**
