Cloth - Rework the forced teleport after initialization to prevent any translation impulse on the first simulated frame.

diff --git a/Engine/Plugins/Experimental/ChaosCloth/Source/Chaos/Private/ChaosCloth/ChaosClothingSimulation.cpp b/Engine/Plugins/Experimental/ChaosCloth/Source/Chaos/Private/ChaosCloth/ChaosClothingSimulation.cpp
--- a/Engine/Plugins/Experimental/ChaosCloth/Source/Chaos/Private/ChaosCloth/ChaosClothingSimulation.cpp
+++ b/Engine/Plugins/Experimental/ChaosCloth/Source/Chaos/Private/ChaosCloth/ChaosClothingSimulation.cpp
@@ -520,7 +520,7 @@ void FClothingSimulation::Simulate(IClothingSimulationContext* InContext)
 #else
 		const bool bNeedsReset = (Context->TeleportMode == EClothingTeleportMode::TeleportAndReset);
 #endif
-		const bool bNeedsTeleport = (Context->TeleportMode > EClothingTeleportMode::None);
+		const bool bNeedsTeleport = (Context->TeleportMode > EClothingTeleportMode::None) || !SimulationTime;
 		bIsTeleported = bNeedsTeleport;
 
 		// Update Solver animatable parameters
diff --git a/Engine/Source/Editor/Persona/Private/AnimationEditorViewportClient.cpp b/Engine/Source/Editor/Persona/Private/AnimationEditorViewportClient.cpp
--- a/Engine/Source/Editor/Persona/Private/AnimationEditorViewportClient.cpp
+++ b/Engine/Source/Editor/Persona/Private/AnimationEditorViewportClient.cpp
@@ -1050,7 +1050,7 @@ FText FAnimationViewportClient::GetDisplayInfo(bool bDisplayAllInfo) const
 		}
 		if (ClothingSimulation->IsTeleported())
 		{
-			TextValue = ConcatenateLine(TextValue, LOCTEXT("IsTeleported", "*** Warning ***: Max Delta Time Teleport!"));
+			TextValue = ConcatenateLine(TextValue, LOCTEXT("IsTeleported", "Simulation Teleport Activated"));
 		}
 	}
 
diff --git a/Engine/Source/Runtime/ClothingSystemRuntimeCommon/Private/ClothingSimulation.cpp b/Engine/Source/Runtime/ClothingSystemRuntimeCommon/Private/ClothingSimulation.cpp
--- a/Engine/Source/Runtime/ClothingSystemRuntimeCommon/Private/ClothingSimulation.cpp
+++ b/Engine/Source/Runtime/ClothingSystemRuntimeCommon/Private/ClothingSimulation.cpp
@@ -55,12 +55,6 @@ void FClothingSimulationContextCommon::Fill(const USkeletalMeshComponent* InComp
 	FillTeleportMode(InComponent, InDeltaSeconds, InMaxPhysicsDelta);
 	FillMaxDistanceScale(InComponent);
 
-	// Cloth fills out the reference pose upon initialization, so flag it as teleporting to avoid destabilizing it if another system was animating it before cloth took over
-	if (bIsInitialization && (TeleportMode != EClothingTeleportMode::TeleportAndReset))
-	{
-		TeleportMode = EClothingTeleportMode::Teleport;
-	}
-
 	PredictedLod = InComponent->GetPredictedLODLevel();
 }
 
@@ -121,13 +115,6 @@ void FClothingSimulationContextCommon::FillBoneTransforms(const USkeletalMeshCom
 	}
 }
 
-void FClothingSimulationContextCommon::FillRefToLocals(const USkeletalMeshComponent* InComponent)
-{
-	// Deprecated version always fills RefToLocals with current animation results instead of using reference pose on initialization
-	const bool bIsInitialization = false;
-	FillRefToLocals(InComponent, bIsInitialization);
-}
-
 void FClothingSimulationContextCommon::FillRefToLocals(const USkeletalMeshComponent* InComponent, bool bIsInitialization)
 {
 	RefToLocals.Reset();
diff --git a/Engine/Source/Runtime/ClothingSystemRuntimeCommon/Public/ClothingSimulation.h b/Engine/Source/Runtime/ClothingSystemRuntimeCommon/Public/ClothingSimulation.h
--- a/Engine/Source/Runtime/ClothingSystemRuntimeCommon/Public/ClothingSimulation.h
+++ b/Engine/Source/Runtime/ClothingSystemRuntimeCommon/Public/ClothingSimulation.h
@@ -32,8 +32,6 @@ public:
 protected:
 	// Default fill behavior as expected to be used by every simulation
 	virtual void FillBoneTransforms(const USkeletalMeshComponent* InComponent);
-	UE_DEPRECATED(4.27, "Use the version with bIsInitialization instead.")
-	virtual void FillRefToLocals(const USkeletalMeshComponent* InComponent);
 	virtual void FillRefToLocals(const USkeletalMeshComponent* InComponent, bool bIsInitialization);
 	virtual void FillComponentToWorld(const USkeletalMeshComponent* InComponent);
 	virtual void FillWorldGravity(const USkeletalMeshComponent* InComponent);
