ClothingAssetUtils::GetMeshClothingAssetBindings Should not have a fallback on the render data. In case we have a scope that unbind/bind the cloth assets, the code in between can also unbind/bind the cloth assets, it has to do nothing in that case. Removing the fallback make the nested call behave correctly.

diff --git a/Engine/Source/Developer/SkeletalMeshUtilitiesCommon/Private/LODUtilities.cpp b/Engine/Source/Developer/SkeletalMeshUtilitiesCommon/Private/LODUtilities.cpp
--- a/Engine/Source/Developer/SkeletalMeshUtilitiesCommon/Private/LODUtilities.cpp
+++ b/Engine/Source/Developer/SkeletalMeshUtilitiesCommon/Private/LODUtilities.cpp
@@ -2949,7 +2949,7 @@ void FLODUtilities::UnbindClothingAndBackup(USkeletalMesh* SkeletalMesh, TArray<
 	}
 	FSkeletalMeshLODModel& LODModel = SkeletalMesh->GetImportedModel()->LODModels[LODIndex];
 	//Store the clothBinding
-	ClothingAssetUtils::GetMeshClothingAssetBindings(SkeletalMesh, ClothingBindings, LODIndex);
+	ClothingAssetUtils::GetAllLodMeshClothingAssetBindings(SkeletalMesh, ClothingBindings, LODIndex);
 	//Unbind the Cloth for this LOD before we reduce it, we will put back the cloth after the reduction, if it still match the sections
 	for (ClothingAssetUtils::FClothingAssetMeshBinding& Binding : ClothingBindings)
 	{
diff --git a/Engine/Source/Editor/ClothingSystemEditor/Private/ClothingAssetFactory.cpp b/Engine/Source/Editor/ClothingSystemEditor/Private/ClothingAssetFactory.cpp
--- a/Engine/Source/Editor/ClothingSystemEditor/Private/ClothingAssetFactory.cpp
+++ b/Engine/Source/Editor/ClothingSystemEditor/Private/ClothingAssetFactory.cpp
@@ -170,7 +170,7 @@ UClothingAssetBase* UClothingAssetFactory::Reimport(const FString& Filename, USk
 			AssetName = *FPaths::GetBaseFilename(Filename);
 
 			TArray<ClothingAssetUtils::FClothingAssetMeshBinding> AssetBindings;
-			ClothingAssetUtils::GetMeshClothingAssetBindings(TargetMesh, AssetBindings);
+			ClothingAssetUtils::GetAllMeshClothingAssetBindings(TargetMesh, AssetBindings);
 
 			OldClothingAsset->UnbindFromSkeletalMesh(TargetMesh);
 
diff --git a/Engine/Source/Runtime/ClothingSystemRuntimeCommon/Private/ClothingAsset.cpp b/Engine/Source/Runtime/ClothingSystemRuntimeCommon/Private/ClothingAsset.cpp
--- a/Engine/Source/Runtime/ClothingSystemRuntimeCommon/Private/ClothingAsset.cpp
+++ b/Engine/Source/Runtime/ClothingSystemRuntimeCommon/Private/ClothingAsset.cpp
@@ -40,6 +40,7 @@ DEFINE_LOG_CATEGORY(LogClothingAsset)
 // ClothingAssetUtils
 //==============================================================================
 
+//Deprecated function
 void ClothingAssetUtils::GetMeshClothingAssetBindings(
 	USkeletalMesh* InSkelMesh, 
 	TArray<FClothingAssetMeshBinding>& OutBindings)
@@ -59,7 +60,9 @@ void ClothingAssetUtils::GetMeshClothingAssetBindings(
 			if (InSkelMesh->GetImportedModel()->LODModels[LODIndex].HasClothData())
 			{
 				TArray<FClothingAssetMeshBinding> LodBindings;
+PRAGMA_DISABLE_DEPRECATION_WARNINGS
 				GetMeshClothingAssetBindings(InSkelMesh, LodBindings, LODIndex);
+PRAGMA_ENABLE_DEPRECATION_WARNINGS
 				OutBindings.Append(LodBindings);
 			}
 		}
@@ -78,13 +81,16 @@ void ClothingAssetUtils::GetMeshClothingAssetBindings(
 		for (int32 LodIndex = 0; LodIndex < NumLods; ++LodIndex)
 		{
 			TArray<FClothingAssetMeshBinding> LodBindings;
+PRAGMA_DISABLE_DEPRECATION_WARNINGS
 			GetMeshClothingAssetBindings(InSkelMesh, LodBindings, LodIndex);
+PRAGMA_ENABLE_DEPRECATION_WARNINGS
 
 			OutBindings.Append(LodBindings);
 		}
 	}
 }
 
+//Deprecated function
 void ClothingAssetUtils::GetMeshClothingAssetBindings(
 	USkeletalMesh* InSkelMesh, 
 	TArray<FClothingAssetMeshBinding>& OutBindings, 
@@ -164,6 +170,65 @@ void ClothingAssetUtils::GetMeshClothingAssetBindings(
 }
 
 #if WITH_EDITOR
+void ClothingAssetUtils::GetAllMeshClothingAssetBindings(const USkeletalMesh* SkeletalMesh, TArray<FClothingAssetMeshBinding>& OutBindings)
+{
+	OutBindings.Empty();
+	if (!SkeletalMesh)
+	{
+		return;
+	}
+	if (const FSkeletalMeshModel* MeshModel = SkeletalMesh->GetImportedModel())
+	{
+		int32 LODNum = MeshModel->LODModels.Num();
+		for (int32 LODIndex = 0; LODIndex < LODNum; ++LODIndex)
+		{
+			const FSkeletalMeshLODModel& MeshLODModel = MeshModel->LODModels[LODIndex];
+			if (MeshLODModel.HasClothData())
+			{
+				TArray<FClothingAssetMeshBinding> LodBindings;
+				GetAllLodMeshClothingAssetBindings(SkeletalMesh, LodBindings, LODIndex);
+				OutBindings.Append(LodBindings);
+			}
+		}
+	}
+}
+
+void ClothingAssetUtils::GetAllLodMeshClothingAssetBindings(const USkeletalMesh* SkeletalMesh, TArray<FClothingAssetMeshBinding>& OutBindings, int32 InLodIndex)
+{
+	OutBindings.Empty();
+	if (!SkeletalMesh && SkeletalMesh->GetImportedModel())
+	{
+		return;
+	}
+	const FSkeletalMeshModel* MeshModel = SkeletalMesh->GetImportedModel();
+	if (!MeshModel || !MeshModel->LODModels.IsValidIndex(InLodIndex))
+	{
+		return;
+	}
+	const FSkeletalMeshLODModel& MeshLODModel = MeshModel->LODModels[InLodIndex];
+
+	if (MeshLODModel.HasClothData())
+	{
+		TArray<FClothingAssetMeshBinding> LodBindings;
+		int32 SectionNum = MeshLODModel.Sections.Num();
+		for (int32 SectionIndex = 0; SectionIndex < SectionNum; ++SectionIndex)
+		{
+			const FSkelMeshSection& Section = MeshLODModel.Sections[SectionIndex];
+			if (Section.HasClothingData())
+			{
+				UClothingAssetBase* ClothingAsset = SkeletalMesh->GetClothingAsset(Section.ClothingData.AssetGuid);
+				FClothingAssetMeshBinding ClothBinding;
+				ClothBinding.Asset = Cast<UClothingAssetCommon>(ClothingAsset);
+				ClothBinding.AssetInternalLodIndex = Section.ClothingData.AssetLodIndex;// InSkelMesh->GetClothingAssetIndex(Section.ClothingData.AssetGuid);
+				check(ClothBinding.AssetInternalLodIndex == Section.ClothingData.AssetLodIndex);
+				ClothBinding.LODIndex = InLodIndex;
+				ClothBinding.SectionIndex = SectionIndex;
+				OutBindings.Add(ClothBinding);
+			}
+		}
+	}
+}
+
 void ClothingAssetUtils::ClearSectionClothingData(FSkelMeshSection& InSection)
 {
 	InSection.ClothingData.AssetGuid = FGuid();
diff --git a/Engine/Source/Runtime/ClothingSystemRuntimeCommon/Public/ClothingAsset.h b/Engine/Source/Runtime/ClothingSystemRuntimeCommon/Public/ClothingAsset.h
--- a/Engine/Source/Runtime/ClothingSystemRuntimeCommon/Public/ClothingAsset.h
+++ b/Engine/Source/Runtime/ClothingSystemRuntimeCommon/Public/ClothingAsset.h
@@ -20,7 +20,7 @@ namespace ClothingAssetUtils
 	/**
 	 * Helper struct to hold binding information on a clothing asset, used to 
 	 * enumerate all of the bindings on a skeletal mesh with 
-	 * \c GetMeshClothingAssetBindings() below.
+	 * \c GetAllMeshClothingAssetBindings() below.
 	 */
 	struct FClothingAssetMeshBinding
 	{
@@ -35,6 +35,7 @@ namespace ClothingAssetUtils
 	 * @param InSkelMesh - The skeletal mesh to search
 	 * @param OutBindings - The list of bindings to write to
 	 */
+	UE_DEPRECATED(5.0, "This Handler Is Now Deprecated, use GetAllMeshClothingAssetBindings instead.")
 	void CLOTHINGSYSTEMRUNTIMECOMMON_API 
 	GetMeshClothingAssetBindings(
 		USkeletalMesh* InSkelMesh, 
@@ -43,6 +44,7 @@ namespace ClothingAssetUtils
 	/**
 	 * Similar to above, but only inspects the specified LOD.
 	 */
+	UE_DEPRECATED(5.0, "This Handler Is Now Deprecated, use GetAllLodMeshClothingAssetBindings instead.")
 	void CLOTHINGSYSTEMRUNTIMECOMMON_API 
 	GetMeshClothingAssetBindings(
 		USkeletalMesh* InSkelMesh, 
@@ -50,6 +52,25 @@ namespace ClothingAssetUtils
 		int32 InLodIndex);
 
 #if WITH_EDITOR
+	/**
+	 * Given a skeletal mesh model, find all of the currently bound clothing assets and their binding information
+	 * @param MeshModel - The skeletal mesh imported model to extract binding from
+	 * @param OutBindings - The list of bindings to write to
+	 */
+	void CLOTHINGSYSTEMRUNTIMECOMMON_API
+	GetAllMeshClothingAssetBindings(
+		const USkeletalMesh* InSkelMesh,
+		TArray<FClothingAssetMeshBinding>& OutBindings);
+
+	/**
+	 * Similar to above, but only inspects the specified LOD.
+	 */
+	void CLOTHINGSYSTEMRUNTIMECOMMON_API
+	GetAllLodMeshClothingAssetBindings(
+		const USkeletalMesh* InSkelMesh,
+		TArray<FClothingAssetMeshBinding>& OutBindings,
+		int32 InLodIndex);
+
 	/**
 	 * Clears the clothing tracking struct of a section.
 	 */
