Fix CIS build

diff --git a/Engine/Source/Runtime/ClothingSystemRuntimeCommon/Private/ClothingAsset.cpp b/Engine/Source/Runtime/ClothingSystemRuntimeCommon/Private/ClothingAsset.cpp
--- a/Engine/Source/Runtime/ClothingSystemRuntimeCommon/Private/ClothingAsset.cpp
+++ b/Engine/Source/Runtime/ClothingSystemRuntimeCommon/Private/ClothingAsset.cpp
@@ -196,7 +196,7 @@ void ClothingAssetUtils::GetAllMeshClothingAssetBindings(const USkeletalMesh* Sk
 void ClothingAssetUtils::GetAllLodMeshClothingAssetBindings(const USkeletalMesh* SkeletalMesh, TArray<FClothingAssetMeshBinding>& OutBindings, int32 InLodIndex)
 {
 	OutBindings.Empty();
-	if (!SkeletalMesh && SkeletalMesh->GetImportedModel())
+	if (!SkeletalMesh || !SkeletalMesh->GetImportedModel())
 	{
 		return;
 	}
