Chaos Cloth - Fix the shared config ownership that is preventing some Skeletal Mesh assets from saving. [Edigrate CL 20537489]

diff --git a/Engine/Source/Runtime/ClothingSystemRuntimeCommon/Private/ClothingAsset.cpp b/Engine/Source/Runtime/ClothingSystemRuntimeCommon/Private/ClothingAsset.cpp
--- a/Engine/Source/Runtime/ClothingSystemRuntimeCommon/Private/ClothingAsset.cpp
+++ b/Engine/Source/Runtime/ClothingSystemRuntimeCommon/Private/ClothingAsset.cpp
@@ -1416,16 +1416,15 @@ void UClothingAssetCommon::PropagateSharedConfigs(bool bMigrateSharedConfigToCon
 			}
 		}
 
-		// Migrate the common shared configs' deprecated parameters to all per cloth configs
-		if (bMigrateSharedConfigToConfig)
+		// Migrate the common shared configs' deprecated parameters to all per cloth configs, and fix the shared config ownership
+		for (const TPair<FName, TObjectPtr<UClothConfigBase>>& ClothSharedConfigItem : ClothConfigs)
 		{
-			// Iterate through all this asset's shared configs
-			for (const auto& ClothSharedConfigItem : ClothConfigs)
+			if (UClothSharedConfigCommon* const ClothSharedConfig = Cast<UClothSharedConfigCommon>(ClothSharedConfigItem.Value))
 			{
-				if (const UClothSharedConfigCommon* const ClothSharedConfig = Cast<UClothSharedConfigCommon>(ClothSharedConfigItem.Value))
+				// Migrate from this shared config to non shared configs if needed
+				if (bMigrateSharedConfigToConfig)
 				{
-					// Iterate through all this asset's configs, and migrate from the shared ones
-					for (const auto& ClothConfigItem : ClothConfigs)
+					for (const TPair<FName, TObjectPtr<UClothConfigBase>>& ClothConfigItem : ClothConfigs)
 					{
 						if (Cast<UClothSharedConfigCommon>(ClothConfigItem.Value))
 						{
@@ -1437,6 +1436,11 @@ void UClothingAssetCommon::PropagateSharedConfigs(bool bMigrateSharedConfigToCon
 						}
 					}
 				}
+				// Fix the shared config outer if it is still a common asset (the config must belong to the skeletal mesh, as it is shared between assets)
+				if (ClothSharedConfig->GetOuter() != SkeletalMesh)
+				{
+					ClothSharedConfig->Rename(nullptr, SkeletalMesh, REN_DontCreateRedirectors | REN_ForceNoResetLoaders | REN_NonTransactional);
+				}
 			}
 		}
 	}
