Do not PreEditChange skeletal mesh while building in async task

diff --git a/Engine/Source/Runtime/ClothingSystemRuntimeCommon/Private/ClothingAsset.cpp b/Engine/Source/Runtime/ClothingSystemRuntimeCommon/Private/ClothingAsset.cpp
--- a/Engine/Source/Runtime/ClothingSystemRuntimeCommon/Private/ClothingAsset.cpp
+++ b/Engine/Source/Runtime/ClothingSystemRuntimeCommon/Private/ClothingAsset.cpp
@@ -779,7 +779,11 @@ void UClothingAssetCommon::UnbindFromSkeletalMesh(
 			FSkelMeshSection& Section = LodModel.Sections[SectionIdx];
 			if(Section.HasClothingData() && Section.ClothingData.AssetGuid == AssetGuid)
 			{
-				InSkelMesh->PreEditChange(nullptr);
+				// Don't do this when in async task, this should have been done before initiating the build
+				if (IsInGameThread())
+				{
+					InSkelMesh->PreEditChange(nullptr);
+				}
 
 				// Log the LOD bias mapping data that will be removed through the call to ClearSectionClothingData
 				for (int32 LodIndex = 0; LodIndex < InMeshLodIndex; ++LodIndex)
