Cloth - Fixed a potential division by zero with a 0 delta time while setting up the context. Also the solver velocity scale is now set to 0 when any cloth teleport mode is activated to match the internal solver behavior.

diff --git a/Engine/Source/Runtime/ClothingSystemRuntimeCommon/Private/ClothingSimulation.cpp b/Engine/Source/Runtime/ClothingSystemRuntimeCommon/Private/ClothingSimulation.cpp
--- a/Engine/Source/Runtime/ClothingSystemRuntimeCommon/Private/ClothingSimulation.cpp
+++ b/Engine/Source/Runtime/ClothingSystemRuntimeCommon/Private/ClothingSimulation.cpp
@@ -167,8 +167,8 @@ void FClothingSimulationContextCommon::FillTeleportMode(const USkeletalMeshCompo
 		EClothingTeleportMode::Teleport :
 		InComponent->ClothTeleportMode;
 
-	VelocityScale = (TeleportMode == EClothingTeleportMode::None) ?
-		FMath::Min(InDeltaSeconds, InMaxPhysicsDelta) / InDeltaSeconds : 1.f;
+	VelocityScale = (TeleportMode == EClothingTeleportMode::None && InDeltaSeconds > 0.f) ?
+		FMath::Min(InDeltaSeconds, InMaxPhysicsDelta) / InDeltaSeconds : 0.f;
 }
 
 void FClothingSimulationContextCommon::FillMaxDistanceScale(const USkeletalMeshComponent* InComponent)
