Cloth - Deprecate APEX collision from the ClothLODData structure.

diff --git a/Engine/Source/Editor/Persona/Private/PersonaMeshDetails.cpp b/Engine/Source/Editor/Persona/Private/PersonaMeshDetails.cpp
--- a/Engine/Source/Editor/Persona/Private/PersonaMeshDetails.cpp
+++ b/Engine/Source/Editor/Persona/Private/PersonaMeshDetails.cpp
@@ -6248,7 +6248,6 @@ TSharedRef<SUniformGridPanel> FPersonaMeshDetails::MakeClothingDetailsWidget(int
 		{
 			const FClothLODDataCommon& LodData = Asset->LodData[LODIndex];
 			const FClothPhysicalMeshData& PhysMeshData = LodData.PhysicalMeshData;
-			const FClothCollisionData& CollisionData = LodData.CollisionData;
 
 			Grid->AddSlot(0, RowNumber)
 				.HAlign(HAlign_Center)
@@ -6314,22 +6313,6 @@ TSharedRef<SUniformGridPanel> FPersonaMeshDetails::MakeClothingDetailsWidget(int
 				.Text(FText::AsNumber(PhysMeshData.MaxBoneWeights))
 				];
 
-			Grid->AddSlot(4, RowNumber)
-				.HAlign(HAlign_Center)
-				[
-					SNew(STextBlock)
-					.Font(DetailFontInfo)
-				.Text(LOCTEXT("NumBoneSpheres", "Spheres"))
-				];
-
-			Grid->AddSlot(4, RowNumber + 1)
-				.HAlign(HAlign_Center)
-				[
-					SNew(STextBlock)
-					.Font(DetailFontInfo)
-				.Text(FText::AsNumber(CollisionData.Spheres.Num()))
-				];
-
 			RowNumber += 2;
 		}
 		else
diff --git a/Engine/Source/Runtime/ClothingSystemRuntimeCommon/Private/ClothLODData.cpp b/Engine/Source/Runtime/ClothingSystemRuntimeCommon/Private/ClothLODData.cpp
--- a/Engine/Source/Runtime/ClothingSystemRuntimeCommon/Private/ClothLODData.cpp
+++ b/Engine/Source/Runtime/ClothingSystemRuntimeCommon/Private/ClothLODData.cpp
@@ -54,22 +54,13 @@ bool FClothLODDataCommon::Serialize(FArchive& Ar)
 		}
 		ParameterMasks_DEPRECATED.Empty();
 
-		// Migrate Convex Planes from the legacy Apex collision (used for box collision)
-		// Note: This code is not enough to get Apex convex working (it's missing surface points and face indices),
-		//       but it is better to keep the data alive for now in case the Apex deprecation causes any issues.
-		for (FClothCollisionPrim_Convex& Convex : CollisionData.Convexes)
-		{
-			const int32 NumDeprecatedPlanes = Convex.Planes_DEPRECATED.Num();
-			if (NumDeprecatedPlanes)
-			{
-				Convex.Faces.SetNum(NumDeprecatedPlanes);
-				for (int32 FaceIndex = 0; FaceIndex < NumDeprecatedPlanes; ++FaceIndex)
-				{
-					Convex.Faces[FaceIndex].Plane = Convex.Planes_DEPRECATED[FaceIndex];
-				}
-				Convex.Planes_DEPRECATED.Empty();
-			}
-		}
+		// Remove deprecated Apex collisions from the LOD data
+PRAGMA_DISABLE_DEPRECATION_WARNINGS
+		CollisionData.Spheres.Empty();
+		CollisionData.SphereConnections.Empty();
+		CollisionData.Convexes.Empty();
+		CollisionData.Boxes.Empty();
+PRAGMA_ENABLE_DEPRECATION_WARNINGS
 	}
 #endif // WITH_EDITORONLY_DATA
 	return true;
diff --git a/Engine/Source/Runtime/ClothingSystemRuntimeCommon/Private/ClothLODData_Legacy.cpp b/Engine/Source/Runtime/ClothingSystemRuntimeCommon/Private/ClothLODData_Legacy.cpp
--- a/Engine/Source/Runtime/ClothingSystemRuntimeCommon/Private/ClothLODData_Legacy.cpp
+++ b/Engine/Source/Runtime/ClothingSystemRuntimeCommon/Private/ClothLODData_Legacy.cpp
@@ -69,12 +69,6 @@ void UClothLODDataCommon_Legacy::MigrateTo(FClothLODDataCommon& LodData)
 		LodData.PhysicalMeshData.MigrateFrom(ClothPhysicalMeshData);
 	}
 
-	// Migrate collision
-	LodData.CollisionData.Spheres = MoveTemp(CollisionData.Spheres);
-	LodData.CollisionData.SphereConnections = MoveTemp(CollisionData.SphereConnections);
-	LodData.CollisionData.Convexes = MoveTemp(CollisionData.Convexes);
-	LodData.CollisionData.Boxes = MoveTemp(CollisionData.Boxes);
-
 #if WITH_EDITORONLY_DATA
 	// Migrate editor maps
 	LodData.PointWeightMaps = MoveTemp(ParameterMasks);
diff --git a/Engine/Source/Runtime/ClothingSystemRuntimeCommon/Public/ClothLODData.h b/Engine/Source/Runtime/ClothingSystemRuntimeCommon/Public/ClothLODData.h
--- a/Engine/Source/Runtime/ClothingSystemRuntimeCommon/Public/ClothLODData.h
+++ b/Engine/Source/Runtime/ClothingSystemRuntimeCommon/Public/ClothLODData.h
@@ -11,12 +11,22 @@ struct CLOTHINGSYSTEMRUNTIMECOMMON_API FClothLODDataCommon
 {
 	GENERATED_BODY()
 
+PRAGMA_DISABLE_DEPRECATION_WARNINGS  // For CollisionData
+	FClothLODDataCommon() = default;
+	FClothLODDataCommon(const FClothLODDataCommon&) = default;
+	FClothLODDataCommon(FClothLODDataCommon&&) = default;
+	~FClothLODDataCommon() = default;
+	FClothLODDataCommon& operator=(const FClothLODDataCommon&) = default;
+	FClothLODDataCommon& operator=(FClothLODDataCommon&&) = default;
+PRAGMA_ENABLE_DEPRECATION_WARNINGS
+
 	// Raw phys mesh data
 	UPROPERTY(EditAnywhere, Category = SimMesh)
 	FClothPhysicalMeshData PhysicalMeshData;
 
 	// Collision primitive and convex data for clothing collisions
-	UPROPERTY(EditAnywhere, Category = Collision)
+	UE_DEPRECATED(5.2, "This property is no longer supported. Use Physics Asset instead.")
+	UPROPERTY(EditAnywhere, Category = Collision, Meta = (DeprecatedProperty, DeprecationMessage = "This property is no longer supported. Use Physics Asset instead."))
 	FClothCollisionData CollisionData;
 
 	// Whether to use multiple triangles to interpolate from simulated cloth mesh to render mesh
