Chaos Cloth - Fixed crash when undoing mask targets.

diff --git a/Engine/Source/Runtime/ClothingSystemRuntimeCommon/Private/ClothingAsset.cpp b/Engine/Source/Runtime/ClothingSystemRuntimeCommon/Private/ClothingAsset.cpp
--- a/Engine/Source/Runtime/ClothingSystemRuntimeCommon/Private/ClothingAsset.cpp
+++ b/Engine/Source/Runtime/ClothingSystemRuntimeCommon/Private/ClothingAsset.cpp
@@ -955,6 +955,45 @@ void UClothingAssetCommon::BuildLodTransitionData()
 	}
 }
 
+void UClothingAssetCommon::PreEditUndo()
+{
+	Super::PreEditUndo();
+
+	// Stop the simulation
+	if (const USkeletalMesh* const OwnerMesh = Cast<USkeletalMesh>(GetOuter()))
+	{
+		for (TObjectIterator<USkeletalMeshComponent> It; It; ++It)
+		{
+			if (USkeletalMeshComponent* const Component = *It)
+			{
+				if (Component->GetSkeletalMeshAsset() == OwnerMesh)
+				{
+					Component->ReleaseAllClothingResources();
+				}
+			}
+		}
+	}
+}
+
+void UClothingAssetCommon::PostEditUndo()
+{
+	Super::PostEditUndo();
+
+	// Resume the simulation
+	if (const USkeletalMesh* const OwnerMesh = Cast<USkeletalMesh>(GetOuter()))
+	{
+		for (TObjectIterator<USkeletalMeshComponent> It; It; ++It)
+		{
+			if (USkeletalMeshComponent* const Component = *It)
+			{
+				if (Component->GetSkeletalMeshAsset() == OwnerMesh)
+				{
+					Component->RecreateClothingActors();
+				}
+			}
+		}
+	}
+}
 #endif // WITH_EDITOR
 
 void UClothingAssetCommon::RefreshBoneMapping(USkeletalMesh* InSkelMesh)
diff --git a/Engine/Source/Runtime/ClothingSystemRuntimeCommon/Public/ClothingAsset.h b/Engine/Source/Runtime/ClothingSystemRuntimeCommon/Public/ClothingAsset.h
--- a/Engine/Source/Runtime/ClothingSystemRuntimeCommon/Public/ClothingAsset.h
+++ b/Engine/Source/Runtime/ClothingSystemRuntimeCommon/Public/ClothingAsset.h
@@ -157,6 +157,17 @@ public:
 	 */
 	void BuildLodTransitionData();
 
+	//~ Begin UObject interface
+	/**
+	 * Stop any simulation from using this asset.
+	 */
+	virtual void PreEditUndo() override;
+
+	/**
+	 * Restart simulation using this asset after undo change.
+	 */
+	virtual void PostEditUndo() override;
+	//~ End UObject interface
 #endif // WITH_EDITOR
 
 	/**
