Fix ClothSimulation cache data deprecation warnings when building with clang.

diff --git a/Engine/Source/Runtime/ClothingSystemRuntimeCommon/Public/ClothingSimulation.h b/Engine/Source/Runtime/ClothingSystemRuntimeCommon/Public/ClothingSimulation.h
--- a/Engine/Source/Runtime/ClothingSystemRuntimeCommon/Public/ClothingSimulation.h
+++ b/Engine/Source/Runtime/ClothingSystemRuntimeCommon/Public/ClothingSimulation.h
@@ -22,6 +22,11 @@ class CLOTHINGSYSTEMRUNTIMECOMMON_API FClothingSimulationContextCommon : public
 {
 public:
 	FClothingSimulationContextCommon();
+	PRAGMA_DISABLE_DEPRECATION_WARNINGS
+	FClothingSimulationContextCommon(const FClothingSimulationContextCommon&) = default;
+	FClothingSimulationContextCommon& operator=(const FClothingSimulationContextCommon&) = default;
+	PRAGMA_ENABLE_DEPRECATION_WARNINGS
+
 	virtual ~FClothingSimulationContextCommon() override;
 
 	UE_DEPRECATED(4.27, "Use the version with bIsInitialization instead.")
