Fixed mismatched copy constructor/assignment operator in ClothingMeshUtils::FVertexParameterMapper.

diff --git a/Engine/Source/Runtime/ClothingSystemRuntimeCommon/Public/Utils/ClothingMeshUtils.h b/Engine/Source/Runtime/ClothingSystemRuntimeCommon/Public/Utils/ClothingMeshUtils.h
--- a/Engine/Source/Runtime/ClothingSystemRuntimeCommon/Public/Utils/ClothingMeshUtils.h
+++ b/Engine/Source/Runtime/ClothingSystemRuntimeCommon/Public/Utils/ClothingMeshUtils.h
@@ -255,8 +255,7 @@ namespace ClothingMeshUtils
 	class CLOTHINGSYSTEMRUNTIMECOMMON_API FVertexParameterMapper
 	{
 	public:
-		FVertexParameterMapper() = delete;
-		FVertexParameterMapper(const FVertexParameterMapper& Other) = delete;
+		UE_NONCOPYABLE(FVertexParameterMapper)
 
 		FVertexParameterMapper(TConstArrayView<FVector3f> InMesh0Positions,
 			TConstArrayView<FVector3f> InMesh0Normals,
