[Engine]

diff --git a/Engine/Source/Runtime/CinematicCamera/Public/CineCameraSettings.h b/Engine/Source/Runtime/CinematicCamera/Public/CineCameraSettings.h
--- a/Engine/Source/Runtime/CinematicCamera/Public/CineCameraSettings.h
+++ b/Engine/Source/Runtime/CinematicCamera/Public/CineCameraSettings.h
@@ -267,16 +267,16 @@ struct FCameraFocusSettings
 	{}
 };
 
-UCLASS(config=Engine, defaultconfig, meta=(DisplayName="Cinematic Camera"))
-class CINEMATICCAMERA_API UCineCameraSettings : public UDeveloperSettings
+UCLASS(config=Engine, defaultconfig, meta=(DisplayName="Cinematic Camera"), MinimalAPI)
+class UCineCameraSettings : public UDeveloperSettings
 {
 	GENERATED_BODY()
 
 public:
-	virtual void PostInitProperties() override;
+	CINEMATICCAMERA_API virtual void PostInitProperties() override;
 
 #if WITH_EDITOR
-	virtual void PostEditChangeProperty(FPropertyChangedEvent& PropertyChangedEvent) override;
+	CINEMATICCAMERA_API virtual void PostEditChangeProperty(FPropertyChangedEvent& PropertyChangedEvent) override;
 #endif
 
 	/** Name of the default lens preset */
@@ -311,24 +311,24 @@ public:
 	UPROPERTY(config, EditAnywhere, BlueprintReadWrite, BlueprintSetter=SetCropPresets, Category=Crop, meta=(TitleProperty=Name))
 	TArray<FNamedPlateCropPreset> CropPresets;
 
-	static TArray<FNamedLensPreset> const& GetLensPresets();
-	static TArray<FNamedFilmbackPreset> const& GetFilmbackPresets();
-	static TArray<FNamedPlateCropPreset> const& GetCropPresets();
+	static CINEMATICCAMERA_API TArray<FNamedLensPreset> const& GetLensPresets();
+	static CINEMATICCAMERA_API TArray<FNamedFilmbackPreset> const& GetFilmbackPresets();
+	static CINEMATICCAMERA_API TArray<FNamedPlateCropPreset> const& GetCropPresets();
 
 	// Gets the Lens settings associated with a given preset name
 	// Returns true if a preset with the given name was found
 	UFUNCTION(BlueprintCallable, Category=Cinematics, meta=(ReturnDisplayName = "Success"))
-	bool GetLensPresetByName(const FString PresetName, FCameraLensSettings& LensSettings);
+	CINEMATICCAMERA_API bool GetLensPresetByName(const FString PresetName, FCameraLensSettings& LensSettings);
 
 	// Gets the Filmback settings associated with a given preset name
 	// Returns true if a preset with the given name was found
 	UFUNCTION(BlueprintCallable, Category=Cinematics, meta=(ReturnDisplayName = "Success"))
-	bool GetFilmbackPresetByName(const FString PresetName, FCameraFilmbackSettings& FilmbackSettings);
+	CINEMATICCAMERA_API bool GetFilmbackPresetByName(const FString PresetName, FCameraFilmbackSettings& FilmbackSettings);
 
 	// Gets the Crop settings associated with a given preset name
 	// Returns true if a preset with the given name was found
 	UFUNCTION(BlueprintCallable, Category=Cinematics, meta=(ReturnDisplayName = "Success"))
-	bool GetCropPresetByName(const FString PresetName, FPlateCropSettings& CropSettings);
+	CINEMATICCAMERA_API bool GetCropPresetByName(const FString PresetName, FPlateCropSettings& CropSettings);
 
 private:
 	UFUNCTION(BlueprintCallable, Category=Cinematics)
diff --git a/Engine/Source/Runtime/ClientPilot/Public/ClientPilotBlackboard.h b/Engine/Source/Runtime/ClientPilot/Public/ClientPilotBlackboard.h
--- a/Engine/Source/Runtime/ClientPilot/Public/ClientPilotBlackboard.h
+++ b/Engine/Source/Runtime/ClientPilot/Public/ClientPilotBlackboard.h
@@ -10,22 +10,22 @@
 
 #include "ClientPilotBlackboard.generated.h"
 
-UCLASS()
-class CLIENTPILOT_API UClientPilotBlackboard : public UObject
+UCLASS(MinimalAPI)
+class UClientPilotBlackboard : public UObject
 {
 GENERATED_BODY()
 public:
-	virtual void InitializeFromProfile(FString CategoryToUse);
+	CLIENTPILOT_API virtual void InitializeFromProfile(FString CategoryToUse);
 
-	void AddOrUpdateValue(FString KeyName, float Value);
-	void AddOrUpdateValue(FString KeyName, int Value);
-	void AddOrUpdateValue(FString KeyName, FString Value);
-	void AddOrUpdateValue(FString KeyName, FVector Value);
+	CLIENTPILOT_API void AddOrUpdateValue(FString KeyName, float Value);
+	CLIENTPILOT_API void AddOrUpdateValue(FString KeyName, int Value);
+	CLIENTPILOT_API void AddOrUpdateValue(FString KeyName, FString Value);
+	CLIENTPILOT_API void AddOrUpdateValue(FString KeyName, FVector Value);
 
-	FString GetStringValue(FString KeyName);
-	int GetIntValue(FString KeyName);
-	float GetFloatValue(FString KeyName);
-	FVector GetVectorValue(FString KeyName);
+	CLIENTPILOT_API FString GetStringValue(FString KeyName);
+	CLIENTPILOT_API int GetIntValue(FString KeyName);
+	CLIENTPILOT_API float GetFloatValue(FString KeyName);
+	CLIENTPILOT_API FVector GetVectorValue(FString KeyName);
 
 	void RemoveKey(FString Key) 
 	{ 
@@ -38,4 +38,4 @@ public:
 
 protected:
 	TMap<FString, FString> Blackboard;
-};
\ No newline at end of file
+};
diff --git a/Engine/Source/Runtime/ClientPilot/Public/ClientPilotBlackboardManager.h b/Engine/Source/Runtime/ClientPilot/Public/ClientPilotBlackboardManager.h
--- a/Engine/Source/Runtime/ClientPilot/Public/ClientPilotBlackboardManager.h
+++ b/Engine/Source/Runtime/ClientPilot/Public/ClientPilotBlackboardManager.h
@@ -12,15 +12,15 @@
 class UClientPilotBlackboard;
 
 
-UCLASS()
-class CLIENTPILOT_API UClientPilotBlackboardManager : public UObject
+UCLASS(MinimalAPI)
+class UClientPilotBlackboardManager : public UObject
 {
 GENERATED_BODY()
 protected:
-	static UClientPilotBlackboardManager * ObjectInstance;
+	static CLIENTPILOT_API UClientPilotBlackboardManager * ObjectInstance;
 
 public:
-	static UClientPilotBlackboardManager * GetInstance();
+	static CLIENTPILOT_API UClientPilotBlackboardManager * GetInstance();
 	UPROPERTY()
 	TObjectPtr<UClientPilotBlackboard> PilotBlackboard;
-};
\ No newline at end of file
+};
diff --git a/Engine/Source/Runtime/ClientPilot/Public/ClientPilotComponent.h b/Engine/Source/Runtime/ClientPilot/Public/ClientPilotComponent.h
--- a/Engine/Source/Runtime/ClientPilot/Public/ClientPilotComponent.h
+++ b/Engine/Source/Runtime/ClientPilot/Public/ClientPilotComponent.h
@@ -10,13 +10,13 @@
 
 class UClientPilotBlackboard;
 
-UCLASS()
-class CLIENTPILOT_API UClientPilotComponent : public UObject
+UCLASS(MinimalAPI)
+class UClientPilotComponent : public UObject
 {
 	GENERATED_BODY()
 public:
-	UClientPilotComponent();
-	UClientPilotBlackboard* GetBlackboardInstance();
-	virtual void ThinkAndAct();
+	CLIENTPILOT_API UClientPilotComponent();
+	CLIENTPILOT_API UClientPilotBlackboard* GetBlackboardInstance();
+	CLIENTPILOT_API virtual void ThinkAndAct();
 
-};
\ No newline at end of file
+};
diff --git a/Engine/Source/Runtime/ClothingSystemRuntimeCommon/Public/ClothConfig.h b/Engine/Source/Runtime/ClothingSystemRuntimeCommon/Public/ClothConfig.h
--- a/Engine/Source/Runtime/ClothingSystemRuntimeCommon/Public/ClothConfig.h
+++ b/Engine/Source/Runtime/ClothingSystemRuntimeCommon/Public/ClothConfig.h
@@ -25,13 +25,13 @@ enum class EClothMassMode : uint8
 };
 
 /** Common configuration base class. */
-UCLASS(Abstract)
-class CLOTHINGSYSTEMRUNTIMECOMMON_API UClothConfigCommon : public UClothConfigBase
+UCLASS(Abstract, MinimalAPI)
+class UClothConfigCommon : public UClothConfigBase
 {
 	GENERATED_BODY()
 public:
-	UClothConfigCommon();
-	virtual ~UClothConfigCommon() override;
+	CLOTHINGSYSTEMRUNTIMECOMMON_API UClothConfigCommon();
+	CLOTHINGSYSTEMRUNTIMECOMMON_API virtual ~UClothConfigCommon() override;
 
 	/** Migrate from the legacy FClothConfig structure. */
 	virtual void MigrateFrom(const FClothConfig_Legacy&) {}
@@ -68,11 +68,11 @@ public:
 };
 
 /** Common shared configuration base class. */
-UCLASS(Abstract)
-class CLOTHINGSYSTEMRUNTIMECOMMON_API UClothSharedConfigCommon : public UClothConfigCommon
+UCLASS(Abstract, MinimalAPI)
+class UClothSharedConfigCommon : public UClothConfigCommon
 {
 	GENERATED_BODY()
 public:
-	UClothSharedConfigCommon();
-	virtual ~UClothSharedConfigCommon() override;
+	CLOTHINGSYSTEMRUNTIMECOMMON_API UClothSharedConfigCommon();
+	CLOTHINGSYSTEMRUNTIMECOMMON_API virtual ~UClothSharedConfigCommon() override;
 };
diff --git a/Engine/Source/Runtime/ClothingSystemRuntimeCommon/Public/ClothConfig_Legacy.h b/Engine/Source/Runtime/ClothingSystemRuntimeCommon/Public/ClothConfig_Legacy.h
--- a/Engine/Source/Runtime/ClothingSystemRuntimeCommon/Public/ClothConfig_Legacy.h
+++ b/Engine/Source/Runtime/ClothingSystemRuntimeCommon/Public/ClothConfig_Legacy.h
@@ -63,7 +63,7 @@ struct FClothConstraintSetup_Legacy
  * Redirected from the now defunct ClothingSystemRuntime module.
  */
 USTRUCT()
-struct CLOTHINGSYSTEMRUNTIMECOMMON_API FClothConfig_Legacy
+struct FClothConfig_Legacy
 {
 	GENERATED_BODY()
 
diff --git a/Engine/Source/Runtime/ClothingSystemRuntimeCommon/Public/ClothLODData.h b/Engine/Source/Runtime/ClothingSystemRuntimeCommon/Public/ClothLODData.h
--- a/Engine/Source/Runtime/ClothingSystemRuntimeCommon/Public/ClothLODData.h
+++ b/Engine/Source/Runtime/ClothingSystemRuntimeCommon/Public/ClothLODData.h
@@ -7,7 +7,7 @@
 
 /** Common Cloth LOD representation for all clothing assets. */
 USTRUCT()
-struct CLOTHINGSYSTEMRUNTIMECOMMON_API FClothLODDataCommon
+struct FClothLODDataCommon
 {
 	GENERATED_BODY()
 
@@ -50,11 +50,11 @@ PRAGMA_ENABLE_DEPRECATION_WARNINGS
 	TArray<FPointWeightMap> PointWeightMaps;
 
 	// Get all available parameter masks for the specified target
-	void GetParameterMasksForTarget(const uint8 InTarget, TArray<FPointWeightMap*>& OutMasks);
+	CLOTHINGSYSTEMRUNTIMECOMMON_API void GetParameterMasksForTarget(const uint8 InTarget, TArray<FPointWeightMap*>& OutMasks);
 #endif // WITH_EDITORONLY_DATA
 #if WITH_EDITOR
 	/** Copy \c ParameterMasks to corresponding targets in \c ClothPhysicalMeshData. */
-	void PushWeightsToMesh();
+	CLOTHINGSYSTEMRUNTIMECOMMON_API void PushWeightsToMesh();
 #endif
 
 	// Skinning data for transitioning from a higher detail LOD to this one
@@ -64,7 +64,7 @@ PRAGMA_ENABLE_DEPRECATION_WARNINGS
 	TArray<FMeshToMeshVertData> TransitionDownSkinData;
 
 	// Custom serialize for transition
-	bool Serialize(FArchive& Ar);
+	CLOTHINGSYSTEMRUNTIMECOMMON_API bool Serialize(FArchive& Ar);
 };
 
 template<>
diff --git a/Engine/Source/Runtime/ClothingSystemRuntimeCommon/Public/ClothLODData_Legacy.h b/Engine/Source/Runtime/ClothingSystemRuntimeCommon/Public/ClothLODData_Legacy.h
--- a/Engine/Source/Runtime/ClothingSystemRuntimeCommon/Public/ClothLODData_Legacy.h
+++ b/Engine/Source/Runtime/ClothingSystemRuntimeCommon/Public/ClothLODData_Legacy.h
@@ -17,13 +17,13 @@ class UClothPhysicalMeshDataBase_Legacy;
  * Redirected from the now defunct ClothingSystemRuntime module.
  */
 USTRUCT()
-struct CLOTHINGSYSTEMRUNTIMECOMMON_API FClothParameterMask_Legacy
+struct FClothParameterMask_Legacy
 {
 	GENERATED_BODY();
 
-	FClothParameterMask_Legacy();
+	CLOTHINGSYSTEMRUNTIMECOMMON_API FClothParameterMask_Legacy();
 
-	void MigrateTo(FPointWeightMap& Weights);
+	CLOTHINGSYSTEMRUNTIMECOMMON_API void MigrateTo(FPointWeightMap& Weights);
 
 	/** Name of the mask, mainly for users to differentiate */
 	UPROPERTY()
@@ -54,13 +54,13 @@ struct CLOTHINGSYSTEMRUNTIMECOMMON_API FClothParameterMask_Legacy
  * Deprecated, legacy definition kept for backward compatibility only.
  * Use FClothLODDataCommon instead.
  */
-UCLASS()
-class CLOTHINGSYSTEMRUNTIMECOMMON_API UClothLODDataCommon_Legacy : public UObject
+UCLASS(MinimalAPI)
+class UClothLODDataCommon_Legacy : public UObject
 {
 	GENERATED_BODY()
 public:
-	UClothLODDataCommon_Legacy(const FObjectInitializer& Init);
-	virtual ~UClothLODDataCommon_Legacy();
+	CLOTHINGSYSTEMRUNTIMECOMMON_API UClothLODDataCommon_Legacy(const FObjectInitializer& Init);
+	CLOTHINGSYSTEMRUNTIMECOMMON_API virtual ~UClothLODDataCommon_Legacy();
 
 	// Deprecated, use ClothPhysicalMeshData instead
 	UPROPERTY()
@@ -87,11 +87,11 @@ public:
 	TArray<FMeshToMeshVertData> TransitionDownSkinData;
 
 	// Custom serialize for transition
-	virtual void Serialize(FArchive& Ar) override;
+	CLOTHINGSYSTEMRUNTIMECOMMON_API virtual void Serialize(FArchive& Ar) override;
 
 	// Migrate deprecated properties
-	virtual void PostLoad() override;
+	CLOTHINGSYSTEMRUNTIMECOMMON_API virtual void PostLoad() override;
 
 	// Migrate this deprecated UObject class to the structure format (called by UClothingAssetCommon::PostLoad())
-	void MigrateTo(struct FClothLODDataCommon& LodData);
+	CLOTHINGSYSTEMRUNTIMECOMMON_API void MigrateTo(struct FClothLODDataCommon& LodData);
 };
diff --git a/Engine/Source/Runtime/ClothingSystemRuntimeCommon/Public/ClothPhysicalMeshData.h b/Engine/Source/Runtime/ClothingSystemRuntimeCommon/Public/ClothPhysicalMeshData.h
--- a/Engine/Source/Runtime/ClothingSystemRuntimeCommon/Public/ClothPhysicalMeshData.h
+++ b/Engine/Source/Runtime/ClothingSystemRuntimeCommon/Public/ClothPhysicalMeshData.h
@@ -22,45 +22,45 @@ template <typename T> struct TObjectPtr;
 
 /** Spatial simulation data for a mesh. */
 USTRUCT()
-struct CLOTHINGSYSTEMRUNTIMECOMMON_API FClothPhysicalMeshData
+struct FClothPhysicalMeshData
 {
 	GENERATED_BODY()
 
 	/** Construct an empty cloth physical mesh with default common targets. */
-	FClothPhysicalMeshData();
+	CLOTHINGSYSTEMRUNTIMECOMMON_API FClothPhysicalMeshData();
 
 	/** Migrate from same, used to migrate LOD data from the UClothLODDataCommon_Legacy. */
-	void MigrateFrom(FClothPhysicalMeshData& ClothPhysicalMeshData);
+	CLOTHINGSYSTEMRUNTIMECOMMON_API void MigrateFrom(FClothPhysicalMeshData& ClothPhysicalMeshData);
 
 	/** Migrate from the legacy physical mesh data class, used to migrate LOD data from the UClothLODDataCommon_Legacy. */
-	void MigrateFrom(UClothPhysicalMeshDataBase_Legacy* ClothPhysicalMeshDataBase);
+	CLOTHINGSYSTEMRUNTIMECOMMON_API void MigrateFrom(UClothPhysicalMeshDataBase_Legacy* ClothPhysicalMeshDataBase);
 
 	/** Reset the default common targets for this cloth physical mesh. */
-	void Reset(const int32 InNumVerts, const int32 InNumIndices);
+	CLOTHINGSYSTEMRUNTIMECOMMON_API void Reset(const int32 InNumVerts, const int32 InNumIndices);
 
 	/** Clear out any default weight maps and delete any other ones. */
-	void ClearWeightMaps();
+	CLOTHINGSYSTEMRUNTIMECOMMON_API void ClearWeightMaps();
 
 	/** Build the self collision indices for the relevant config. */
 	UE_DEPRECATED(5.0, "Use BuildSelfCollisionData(float SelfCollisionRadius) instead.")
-	void BuildSelfCollisionData(const TMap<FName, TObjectPtr<UClothConfigBase>>& ClothConfigs);
+	CLOTHINGSYSTEMRUNTIMECOMMON_API void BuildSelfCollisionData(const TMap<FName, TObjectPtr<UClothConfigBase>>& ClothConfigs);
 
 	/** Build the self collision indices with the specified radius. */
-	void BuildSelfCollisionData(float SelfCollisionRadius);
+	CLOTHINGSYSTEMRUNTIMECOMMON_API void BuildSelfCollisionData(float SelfCollisionRadius);
 
 	/** Recalculate the node inverse masses. */
-	void CalculateInverseMasses();
+	CLOTHINGSYSTEMRUNTIMECOMMON_API void CalculateInverseMasses();
 
 	/** Recalculate the number of influences for the bone data. */
-	void CalculateNumInfluences();
+	CLOTHINGSYSTEMRUNTIMECOMMON_API void CalculateNumInfluences();
 
 	/** Recalculate the long range attachment tethers. */
-	void CalculateTethers(bool bUseEuclideanDistance, bool bUseGeodesicDistance);
+	CLOTHINGSYSTEMRUNTIMECOMMON_API void CalculateTethers(bool bUseEuclideanDistance, bool bUseGeodesicDistance);
 
 	/** Compute vertex normals as unweighted average of incident face normals. If all incident triangles are degenerate,
 	    the vertex normal is assigned a value of FVector3f::XAxisVector. */
 	UE_DEPRECATED(5.0, "This function is no longer part of this API and will soon be removed.")
-	void ComputeFaceAveragedVertexNormals(TArray<FVector3f>& OutNormals) const;
+	CLOTHINGSYSTEMRUNTIMECOMMON_API void ComputeFaceAveragedVertexNormals(TArray<FVector3f>& OutNormals) const;
 
 	/** Retrieve whether a vertex weight array has already been registered. */
 	template<typename T>
diff --git a/Engine/Source/Runtime/ClothingSystemRuntimeCommon/Public/ClothingAsset.h b/Engine/Source/Runtime/ClothingSystemRuntimeCommon/Public/ClothingAsset.h
--- a/Engine/Source/Runtime/ClothingSystemRuntimeCommon/Public/ClothingAsset.h
+++ b/Engine/Source/Runtime/ClothingSystemRuntimeCommon/Public/ClothingAsset.h
@@ -114,13 +114,13 @@ ENUM_CLASS_FLAGS(EClothingCachedDataFlagsCommon);
  * \c CustomData, as well as override the \c AddNewLod() factory to build their 
  * own implementation of \c UClothLODDataBase.
  */
-UCLASS(hidecategories = Object, BlueprintType)
-class CLOTHINGSYSTEMRUNTIMECOMMON_API UClothingAssetCommon : public UClothingAssetBase
+UCLASS(hidecategories = Object, BlueprintType, MinimalAPI)
+class UClothingAssetCommon : public UClothingAssetBase
 {
 	GENERATED_BODY()
 public:
 
-	UClothingAssetCommon(const FObjectInitializer& ObjectInitializer);
+	CLOTHINGSYSTEMRUNTIMECOMMON_API UClothingAssetCommon(const FObjectInitializer& ObjectInitializer);
 
 #if WITH_EDITOR
 
@@ -128,45 +128,45 @@ public:
 	 * Create weights for skinning the render mesh to our simulation mesh, and 
 	 * weights to drive our sim mesh from the skeleton.
 	 */
-	virtual bool BindToSkeletalMesh(USkeletalMesh* InSkelMesh, const int32 InMeshLodIndex, const int32 InSectionIndex, const int32 InAssetLodIndex) override;
+	CLOTHINGSYSTEMRUNTIMECOMMON_API virtual bool BindToSkeletalMesh(USkeletalMesh* InSkelMesh, const int32 InMeshLodIndex, const int32 InSectionIndex, const int32 InAssetLodIndex) override;
 
 	/**
 	 * Helper that invokes \c UnbindFromSkeletalMesh() for each avilable entry in 
 	 * \p InSkelMesh->GetImportedModel()'s LODModel.
 	 */
-	virtual void UnbindFromSkeletalMesh(USkeletalMesh* InSkelMesh) override;
-	virtual void UnbindFromSkeletalMesh(USkeletalMesh* InSkelMesh, const int32 InMeshLodIndex) override;
+	CLOTHINGSYSTEMRUNTIMECOMMON_API virtual void UnbindFromSkeletalMesh(USkeletalMesh* InSkelMesh) override;
+	CLOTHINGSYSTEMRUNTIMECOMMON_API virtual void UnbindFromSkeletalMesh(USkeletalMesh* InSkelMesh, const int32 InMeshLodIndex) override;
 
 	/**
 	 * Update all extra LOD deformer mappings.
 	 * This should be called whenever the raytracing LOD bias has changed.
 	 */
-	virtual void UpdateAllLODBiasMappings(USkeletalMesh* SkeletalMesh) override;
+	CLOTHINGSYSTEMRUNTIMECOMMON_API virtual void UpdateAllLODBiasMappings(USkeletalMesh* SkeletalMesh) override;
 
 	/** 
 	 * Callback envoked after weights have been edited.
 	 * Calls \c PushWeightsToMesh() on each \c ClothLodData, and invalidates cached data. 
 	 * Optionaly recalculate the owner's sections fixed vertex data based on this asset masks, and invalidate the DDC for this asset.
 	 */
-	void ApplyParameterMasks(bool bUpdateFixedVertData = false, bool bInvalidateDerivedDataCache = true);
+	CLOTHINGSYSTEMRUNTIMECOMMON_API void ApplyParameterMasks(bool bUpdateFixedVertData = false, bool bInvalidateDerivedDataCache = true);
 
 	/**
 	 * Builds the LOD transition data.
 	 * When we transition between LODs we skin the incoming mesh to the outgoing mesh
 	 * in exactly the same way the render mesh is skinned to create a smooth swap
 	 */
-	void BuildLodTransitionData();
+	CLOTHINGSYSTEMRUNTIMECOMMON_API void BuildLodTransitionData();
 
 	//~ Begin UObject interface
 	/**
 	 * Stop any simulation from using this asset.
 	 */
-	virtual void PreEditUndo() override;
+	CLOTHINGSYSTEMRUNTIMECOMMON_API virtual void PreEditUndo() override;
 
 	/**
 	 * Restart simulation using this asset after undo change.
 	 */
-	virtual void PostEditUndo() override;
+	CLOTHINGSYSTEMRUNTIMECOMMON_API virtual void PostEditUndo() override;
 	//~ End UObject interface
 #endif // WITH_EDITOR
 
@@ -174,16 +174,16 @@ public:
 	 * Rebuilds the \c UsedBoneIndices array from looking up the entries in the
 	 * \c UsedBoneNames array, in the \p InSkelMesh's reference skeleton.
 	 */
-	virtual void RefreshBoneMapping(USkeletalMesh* InSkelMesh) override;
+	CLOTHINGSYSTEMRUNTIMECOMMON_API virtual void RefreshBoneMapping(USkeletalMesh* InSkelMesh) override;
 
 	/** Calculates the preferred root bone for the simulation. */
-	void CalculateReferenceBoneIndex();
+	CLOTHINGSYSTEMRUNTIMECOMMON_API void CalculateReferenceBoneIndex();
 
 	/** Returns \c true if \p InLodIndex is a valid LOD id (index into \c ClothLodData). */
-	virtual bool IsValidLod(int32 InLodIndex) const override;
+	CLOTHINGSYSTEMRUNTIMECOMMON_API virtual bool IsValidLod(int32 InLodIndex) const override;
 
 	/** Returns the number of valid LOD's (length of the \c ClothLodData array). */
-	virtual int32 GetNumLods() const override;
+	CLOTHINGSYSTEMRUNTIMECOMMON_API virtual int32 GetNumLods() const override;
 
 #if WITH_EDITORONLY_DATA
 	/**
@@ -196,15 +196,15 @@ public:
 	 * Called on the clothing asset when the base data (physical mesh, config etc.)
 	 * has changed, so any intermediate generated data can be regenerated.
 	 */
-	void InvalidateFlaggedCachedData(EClothingCachedDataFlagsCommon Flags);
+	CLOTHINGSYSTEMRUNTIMECOMMON_API void InvalidateFlaggedCachedData(EClothingCachedDataFlagsCommon Flags);
 #endif // WITH_EDITORONLY_DATA
 
 #if WITH_EDITOR
 	/** * Add a new LOD class instance. */
-	virtual int32 AddNewLod() override;
+	CLOTHINGSYSTEMRUNTIMECOMMON_API virtual int32 AddNewLod() override;
 
 	/* Called after changes in any of the asset properties. */
-	virtual void PostEditChangeChainProperty(FPropertyChangedChainEvent& ChainEvent) override;
+	CLOTHINGSYSTEMRUNTIMECOMMON_API virtual void PostEditChangeChainProperty(FPropertyChangedChainEvent& ChainEvent) override;
 #endif // WITH_EDITOR
 
 	/**
@@ -240,16 +240,16 @@ PRAGMA_ENABLE_DEPRECATION_WARNINGS
 	}
 
 	/** Migrate deprecated objects. */
-	virtual void PostLoad() override;
+	CLOTHINGSYSTEMRUNTIMECOMMON_API virtual void PostLoad() override;
 #if WITH_EDITORONLY_DATA
-	static void DeclareConstructClasses(TArray<FTopLevelAssetPath>& OutConstructClasses, const UClass* SpecificSubclass);
+	static CLOTHINGSYSTEMRUNTIMECOMMON_API void DeclareConstructClasses(TArray<FTopLevelAssetPath>& OutConstructClasses, const UClass* SpecificSubclass);
 #endif
 
 	/** Serialize deprecated objects. */
-	virtual void Serialize(FArchive& Ar) override;
+	CLOTHINGSYSTEMRUNTIMECOMMON_API virtual void Serialize(FArchive& Ar) override;
 
 	/** Propagate the shared simulation configs between assets. Called after all cloth assets sharing the same simulation are loaded. */
-	virtual void PostUpdateAllAssets() override;
+	CLOTHINGSYSTEMRUNTIMECOMMON_API virtual void PostUpdateAllAssets() override;
 
 	// The physics asset to extract collisions from when building a simulation.
 	UPROPERTY(EditAnywhere, Category = Config)
@@ -333,12 +333,12 @@ private:
 	// If a config from a different factory exists already, the newly
 	// created config will attempt to initialize its parameters from it.
 	// Return true when at least one config has been added, false otherwise.
-	bool AddClothConfigs();
+	CLOTHINGSYSTEMRUNTIMECOMMON_API bool AddClothConfigs();
 
 	// Propagate the shared simulation configs between assets.
 	// Also migrate all deprecated shared parameters which have been moved to the per cloth configs if required.
 	// Called after a cloth asset is created or loaded.
-	void PropagateSharedConfigs(bool bMigrateSharedConfigToConfig=false);
+	CLOTHINGSYSTEMRUNTIMECOMMON_API void PropagateSharedConfigs(bool bMigrateSharedConfigToConfig=false);
 
 	// Return true when any one of the cloth configs fullfill the predicate.
 	// Used to select which type of data to cache.
@@ -353,13 +353,13 @@ private:
 
 #if WITH_EDITOR
 	// Add extra cloth deformer mappings to cope with a different raytracing LOD than the one currently rendered.
-	void UpdateLODBiasMappings(const USkeletalMesh* SkeletalMesh, int32 UpdatedLODIndex, int32 SectionIndex);
+	CLOTHINGSYSTEMRUNTIMECOMMON_API void UpdateLODBiasMappings(const USkeletalMesh* SkeletalMesh, int32 UpdatedLODIndex, int32 SectionIndex);
 
 	// Clear all defomer section that relies on the specified cloth sim data bound to the specified LOD section 
-	void ClearLODBiasMappings(const USkeletalMesh* SkeletalMesh, int32 UpdatedLODIndex, int32 SectionIndex);
+	CLOTHINGSYSTEMRUNTIMECOMMON_API void ClearLODBiasMappings(const USkeletalMesh* SkeletalMesh, int32 UpdatedLODIndex, int32 SectionIndex);
 
 	// Helper functions used in PostPropertyChangeCb
-	void ReregisterComponentsUsingClothing();
-	void ForEachInteractorUsingClothing(TFunction<void (UClothingSimulationInteractor*)> Func);
+	CLOTHINGSYSTEMRUNTIMECOMMON_API void ReregisterComponentsUsingClothing();
+	CLOTHINGSYSTEMRUNTIMECOMMON_API void ForEachInteractorUsingClothing(TFunction<void (UClothingSimulationInteractor*)> Func);
 #endif // WITH_EDITOR
 };
diff --git a/Engine/Source/Runtime/ClothingSystemRuntimeCommon/Public/ClothingSimulation.h b/Engine/Source/Runtime/ClothingSystemRuntimeCommon/Public/ClothingSimulation.h
--- a/Engine/Source/Runtime/ClothingSystemRuntimeCommon/Public/ClothingSimulation.h
+++ b/Engine/Source/Runtime/ClothingSystemRuntimeCommon/Public/ClothingSimulation.h
@@ -18,33 +18,33 @@ DECLARE_CYCLE_STAT_EXTERN(TEXT("Skin Physics Mesh"), STAT_ClothSkinPhysMesh, STA
 DECLARE_CYCLE_STAT_EXTERN(TEXT("Fill Context"), STAT_ClothFillContext, STATGROUP_Physics, CLOTHINGSYSTEMRUNTIMECOMMON_API);
 
 /** Base simulation data that just about every simulation would need. */
-class CLOTHINGSYSTEMRUNTIMECOMMON_API FClothingSimulationContextCommon : public IClothingSimulationContext
+class FClothingSimulationContextCommon : public IClothingSimulationContext
 {
 public:
-	FClothingSimulationContextCommon();
+	CLOTHINGSYSTEMRUNTIMECOMMON_API FClothingSimulationContextCommon();
 	PRAGMA_DISABLE_DEPRECATION_WARNINGS
 	FClothingSimulationContextCommon(const FClothingSimulationContextCommon&) = default;
 	FClothingSimulationContextCommon& operator=(const FClothingSimulationContextCommon&) = default;
-	PRAGMA_ENABLE_DEPRECATION_WARNINGS
+	CLOTHINGSYSTEMRUNTIMECOMMON_API PRAGMA_ENABLE_DEPRECATION_WARNINGS
 
 	virtual ~FClothingSimulationContextCommon() override;
 
 	UE_DEPRECATED(4.27, "Use the version with bIsInitialization instead.")
-	virtual void Fill(const USkeletalMeshComponent* InComponent, float InDeltaSeconds, float InMaxPhysicsDelta);
+	CLOTHINGSYSTEMRUNTIMECOMMON_API virtual void Fill(const USkeletalMeshComponent* InComponent, float InDeltaSeconds, float InMaxPhysicsDelta);
 
 	// Fill this context using the given skeletal mesh component
-	virtual void Fill(const USkeletalMeshComponent* InComponent, float InDeltaSeconds, float InMaxPhysicsDelta, bool bIsInitialization);
+	CLOTHINGSYSTEMRUNTIMECOMMON_API virtual void Fill(const USkeletalMeshComponent* InComponent, float InDeltaSeconds, float InMaxPhysicsDelta, bool bIsInitialization);
 
 protected:
 	// Default fill behavior as expected to be used by every simulation
-	virtual void FillBoneTransforms(const USkeletalMeshComponent* InComponent);
-	virtual void FillRefToLocals(const USkeletalMeshComponent* InComponent, bool bIsInitialization);
-	virtual void FillComponentToWorld(const USkeletalMeshComponent* InComponent);
-	virtual void FillWorldGravity(const USkeletalMeshComponent* InComponent);
-	virtual void FillWindVelocity(const USkeletalMeshComponent* InComponent);
-	virtual void FillDeltaSeconds(float InDeltaSeconds, float InMaxPhysicsDelta);
-	virtual void FillTeleportMode(const USkeletalMeshComponent* InComponent, float InDeltaSeconds, float InMaxPhysicsDelta);
-	virtual void FillMaxDistanceScale(const USkeletalMeshComponent* InComponent);
+	CLOTHINGSYSTEMRUNTIMECOMMON_API virtual void FillBoneTransforms(const USkeletalMeshComponent* InComponent);
+	CLOTHINGSYSTEMRUNTIMECOMMON_API virtual void FillRefToLocals(const USkeletalMeshComponent* InComponent, bool bIsInitialization);
+	CLOTHINGSYSTEMRUNTIMECOMMON_API virtual void FillComponentToWorld(const USkeletalMeshComponent* InComponent);
+	CLOTHINGSYSTEMRUNTIMECOMMON_API virtual void FillWorldGravity(const USkeletalMeshComponent* InComponent);
+	CLOTHINGSYSTEMRUNTIMECOMMON_API virtual void FillWindVelocity(const USkeletalMeshComponent* InComponent);
+	CLOTHINGSYSTEMRUNTIMECOMMON_API virtual void FillDeltaSeconds(float InDeltaSeconds, float InMaxPhysicsDelta);
+	CLOTHINGSYSTEMRUNTIMECOMMON_API virtual void FillTeleportMode(const USkeletalMeshComponent* InComponent, float InDeltaSeconds, float InMaxPhysicsDelta);
+	CLOTHINGSYSTEMRUNTIMECOMMON_API virtual void FillMaxDistanceScale(const USkeletalMeshComponent* InComponent);
 
 public:
 	// Component space bone transforms of the owning component
@@ -94,18 +94,18 @@ public:
 };
 
 // Base simulation to fill in common data for the base context
-class CLOTHINGSYSTEMRUNTIMECOMMON_API FClothingSimulationCommon : public IClothingSimulation
+class FClothingSimulationCommon : public IClothingSimulation
 {
 public:
-	FClothingSimulationCommon();
-	virtual ~FClothingSimulationCommon();
+	CLOTHINGSYSTEMRUNTIMECOMMON_API FClothingSimulationCommon();
+	CLOTHINGSYSTEMRUNTIMECOMMON_API virtual ~FClothingSimulationCommon();
 
 protected:
 	/** Fills in the base data for a clothing simulation */
-	virtual void FillContext(USkeletalMeshComponent* InComponent, float InDeltaTime, IClothingSimulationContext* InOutContext) override;
+	CLOTHINGSYSTEMRUNTIMECOMMON_API virtual void FillContext(USkeletalMeshComponent* InComponent, float InDeltaTime, IClothingSimulationContext* InOutContext) override;
 
 	/** Fills in the base data for a clothing simulation */
-	virtual void FillContext(USkeletalMeshComponent* InComponent, float InDeltaTime, IClothingSimulationContext* InOutContext, bool bIsInitialization) override;
+	CLOTHINGSYSTEMRUNTIMECOMMON_API virtual void FillContext(USkeletalMeshComponent* InComponent, float InDeltaTime, IClothingSimulationContext* InOutContext, bool bIsInitialization) override;
 
 protected:
 	/** Maximum physics time, incoming deltas will be clamped down to this value on long frames */
diff --git a/Engine/Source/Runtime/ClothingSystemRuntimeCommon/Public/PointWeightMap.h b/Engine/Source/Runtime/ClothingSystemRuntimeCommon/Public/PointWeightMap.h
--- a/Engine/Source/Runtime/ClothingSystemRuntimeCommon/Public/PointWeightMap.h
+++ b/Engine/Source/Runtime/ClothingSystemRuntimeCommon/Public/PointWeightMap.h
@@ -22,7 +22,7 @@ enum class EWeightMapTargetCommon : uint8
  * is then later applied to a phys mesh
  */
 USTRUCT()
-struct CLOTHINGSYSTEMRUNTIMECOMMON_API FPointWeightMap
+struct FPointWeightMap
 {
 	GENERATED_BODY();
 
diff --git a/Engine/Source/Runtime/ClothingSystemRuntimeCommon/Public/Utils/ClothingMeshUtils.h b/Engine/Source/Runtime/ClothingSystemRuntimeCommon/Public/Utils/ClothingMeshUtils.h
--- a/Engine/Source/Runtime/ClothingSystemRuntimeCommon/Public/Utils/ClothingMeshUtils.h
+++ b/Engine/Source/Runtime/ClothingSystemRuntimeCommon/Public/Utils/ClothingMeshUtils.h
@@ -12,7 +12,7 @@ struct FPointWeightMap;
 
 namespace ClothingMeshUtils
 {
-	class CLOTHINGSYSTEMRUNTIMECOMMON_API ClothMeshDesc
+	class ClothMeshDesc
 	{
 	public:
 		ClothMeshDesc(TConstArrayView<FVector3f> InPositions, TConstArrayView<FVector3f> InNormals, TConstArrayView<uint32> InIndices)
@@ -57,19 +57,19 @@ namespace ClothingMeshUtils
 		bool HasAveragedNormals() const { return Normals.Num() == AveragedNormals.Num(); }
 
 		/** Find the distance to the specified triangle from a point. */
-		float DistanceToTriangle(const FVector& Position, int32 TriangleBaseIndex) const;
+		CLOTHINGSYSTEMRUNTIMECOMMON_API float DistanceToTriangle(const FVector& Position, int32 TriangleBaseIndex) const;
 
 		/** Find the closest triangles from the specified point. */
-		TArray<int32> FindCandidateTriangles(const FVector& InPoint, float InTolerance = KINDA_SMALL_NUMBER) const;
+		CLOTHINGSYSTEMRUNTIMECOMMON_API TArray<int32> FindCandidateTriangles(const FVector& InPoint, float InTolerance = KINDA_SMALL_NUMBER) const;
 
 		/**
 		 * Return the max edge length for each edge coming out of a mesh vertex.
 		 * Useful for guiding the search radius when searching for nearest triangles.
 		 */
-		TConstArrayView<float> GetMaxEdgeLengths() const;
+		CLOTHINGSYSTEMRUNTIMECOMMON_API TConstArrayView<float> GetMaxEdgeLengths() const;
 
 	private:
-		void ComputeAveragedNormals();
+		CLOTHINGSYSTEMRUNTIMECOMMON_API void ComputeAveragedNormals();
 
 		struct FClothBvEntry;
 
@@ -252,7 +252,7 @@ namespace ClothingMeshUtils
 	 * Object used to map vertex parameters between two meshes using the
 	 * same barycentric mesh to mesh mapping data we use for clothing
 	 */
-	class CLOTHINGSYSTEMRUNTIMECOMMON_API FVertexParameterMapper
+	class FVertexParameterMapper
 	{
 	public:
 		UE_NONCOPYABLE(FVertexParameterMapper)
@@ -340,7 +340,7 @@ namespace ClothingMeshUtils
 		 * This transfers the values from source to dest by matching the mesh topologies using
 		 * barycentric coordinates of the dest mesh (Mesh0) positions over the source mesh (Mesh1).
 		 */
-		void Map(TConstArrayView<float> Source, TArray<float>& Dest);
+		CLOTHINGSYSTEMRUNTIMECOMMON_API void Map(TConstArrayView<float> Source, TArray<float>& Dest);
 
 	private:
 
@@ -349,7 +349,7 @@ namespace ClothingMeshUtils
 		 * @param OutEmbeddedPositions Embedded version of the original positions, a barycentric coordinate and distance along the normal of the triangle
 		 * @param OutSourceIndices Source index list for the embedded positions, 3 per position to denote the source triangle
 		 */
-		void GenerateEmbeddedPositions(TArray<FVector4>& OutEmbeddedPositions, TArray<int32>& OutSourceIndices);
+		CLOTHINGSYSTEMRUNTIMECOMMON_API void GenerateEmbeddedPositions(TArray<FVector4>& OutEmbeddedPositions, TArray<int32>& OutSourceIndices);
 
 		TConstArrayView<FVector3f> Mesh0Positions;
 		TConstArrayView<FVector3f> Mesh0Normals;
@@ -357,4 +357,4 @@ namespace ClothingMeshUtils
 		TConstArrayView<FVector3f> Mesh1Normals;
 		TConstArrayView<uint32> Mesh1Indices;
 	};
-}
\ No newline at end of file
+}
diff --git a/Engine/Source/Runtime/ClothingSystemRuntimeInterface/Public/ClothCollisionData.h b/Engine/Source/Runtime/ClothingSystemRuntimeInterface/Public/ClothCollisionData.h
--- a/Engine/Source/Runtime/ClothingSystemRuntimeInterface/Public/ClothCollisionData.h
+++ b/Engine/Source/Runtime/ClothingSystemRuntimeInterface/Public/ClothCollisionData.h
@@ -8,13 +8,13 @@
 #include "ClothCollisionData.generated.h"
 
 USTRUCT()
-struct CLOTHINGSYSTEMRUNTIMEINTERFACE_API FClothCollisionData
+struct FClothCollisionData
 {
 	GENERATED_BODY()
 
-	void Reset();
+	CLOTHINGSYSTEMRUNTIMEINTERFACE_API void Reset();
 
-	void Append(const FClothCollisionData& InOther);
+	CLOTHINGSYSTEMRUNTIMEINTERFACE_API void Append(const FClothCollisionData& InOther);
 
 	// Sphere data
 	UPROPERTY(EditAnywhere, Category = Collison)
diff --git a/Engine/Source/Runtime/ClothingSystemRuntimeInterface/Public/ClothCollisionPrim.h b/Engine/Source/Runtime/ClothingSystemRuntimeInterface/Public/ClothCollisionPrim.h
--- a/Engine/Source/Runtime/ClothingSystemRuntimeInterface/Public/ClothCollisionPrim.h
+++ b/Engine/Source/Runtime/ClothingSystemRuntimeInterface/Public/ClothCollisionPrim.h
@@ -69,7 +69,7 @@ struct FClothCollisionPrim_ConvexFace
  *	shape created by the planes combined.
  */
 USTRUCT()
-struct CLOTHINGSYSTEMRUNTIMEINTERFACE_API FClothCollisionPrim_Convex
+struct FClothCollisionPrim_Convex
 {
 	GENERATED_BODY()
 
diff --git a/Engine/Source/Runtime/ClothingSystemRuntimeInterface/Public/ClothConfigBase.h b/Engine/Source/Runtime/ClothingSystemRuntimeInterface/Public/ClothConfigBase.h
--- a/Engine/Source/Runtime/ClothingSystemRuntimeInterface/Public/ClothConfigBase.h
+++ b/Engine/Source/Runtime/ClothingSystemRuntimeInterface/Public/ClothConfigBase.h
@@ -14,40 +14,40 @@
  * Base class for simulator specific simulation controls.
  * Each cloth instance on a skeletal mesh can have a unique cloth config
  */
-UCLASS(Abstract)
-class CLOTHINGSYSTEMRUNTIMEINTERFACE_API UClothConfigBase : public UObject
+UCLASS(Abstract, MinimalAPI)
+class UClothConfigBase : public UObject
 {
 	GENERATED_BODY()
 public:
-	UClothConfigBase();
-	virtual ~UClothConfigBase();
+	CLOTHINGSYSTEMRUNTIMEINTERFACE_API UClothConfigBase();
+	CLOTHINGSYSTEMRUNTIMEINTERFACE_API virtual ~UClothConfigBase();
 
 	/** Return the self collision radius if building self collision indices is required for this config, or 0.f otherwise. */
 	UE_DEPRECATED(5.0, "Use NeedsSelfCollisionData and GetSelfCollisionRadius instead.")
 	virtual float NeedsSelfCollisionIndices() const { return NeedsSelfCollisionData() ? GetSelfCollisionRadius() : 0.f; }
 
 	/** Return wherether to pre-compute self collision data. */
-	virtual bool NeedsSelfCollisionData() const
+	CLOTHINGSYSTEMRUNTIMEINTERFACE_API virtual bool NeedsSelfCollisionData() const
 	PURE_VIRTUAL(UClothConfigBase::NeedsSelfCollisionData, return false;);
 
 	/** Return wherether to pre-compute inverse masses. */
-	virtual bool NeedsInverseMasses() const
+	CLOTHINGSYSTEMRUNTIMEINTERFACE_API virtual bool NeedsInverseMasses() const
 	PURE_VIRTUAL(UClothConfigBase::NeedsInverseMasses, return false;);
 
 	/** Return wherether to pre-compute the influences. */
-	virtual bool NeedsNumInfluences() const
+	CLOTHINGSYSTEMRUNTIMEINTERFACE_API virtual bool NeedsNumInfluences() const
 	PURE_VIRTUAL(UClothConfigBase::NeedsNumInfluences, return false;);
 
 	/** Return wherether to pre-compute the long range attachment tethers. */
-	virtual bool NeedsTethers() const
+	CLOTHINGSYSTEMRUNTIMEINTERFACE_API virtual bool NeedsTethers() const
 	PURE_VIRTUAL(UClothConfigBase::NeedsTethers, return false;);
 
 	/** Return the self collision radius to precomute self collision data. */
-	virtual float GetSelfCollisionRadius() const
+	CLOTHINGSYSTEMRUNTIMEINTERFACE_API virtual float GetSelfCollisionRadius() const
 	PURE_VIRTUAL(UClothConfigBase::GetSelfCollisionRadius, return 0.f;);
 
 	/** Return whether tethers need to be calculated using geodesic distances instead of eclidean. */
-	virtual bool TethersUseGeodesicDistance() const
+	CLOTHINGSYSTEMRUNTIMEINTERFACE_API virtual bool TethersUseGeodesicDistance() const
 	PURE_VIRTUAL(UClothConfigBase::TethersUseGeodesicDistance, return false;);
 };
 
@@ -55,8 +55,8 @@ public:
  * These settings are shared between all instances on a skeletal mesh
  * Deprecated, use UClothConfigBase instead.
  */
-UCLASS(Abstract, Deprecated)
-class CLOTHINGSYSTEMRUNTIMEINTERFACE_API UDEPRECATED_ClothSharedSimConfigBase : public UObject
+UCLASS(Abstract, Deprecated, MinimalAPI)
+class UDEPRECATED_ClothSharedSimConfigBase : public UObject
 {
 	GENERATED_BODY()
 public:
diff --git a/Engine/Source/Runtime/ClothingSystemRuntimeInterface/Public/ClothPhysicalMeshDataBase_Legacy.h b/Engine/Source/Runtime/ClothingSystemRuntimeInterface/Public/ClothPhysicalMeshDataBase_Legacy.h
--- a/Engine/Source/Runtime/ClothingSystemRuntimeInterface/Public/ClothPhysicalMeshDataBase_Legacy.h
+++ b/Engine/Source/Runtime/ClothingSystemRuntimeInterface/Public/ClothPhysicalMeshDataBase_Legacy.h
@@ -22,26 +22,26 @@ struct FColor;
  *
  * Created curing asset import or created from a skeletal mesh.
  */
-UCLASS()
-class CLOTHINGSYSTEMRUNTIMEINTERFACE_API UClothPhysicalMeshDataBase_Legacy : public UObject
+UCLASS(MinimalAPI)
+class UClothPhysicalMeshDataBase_Legacy : public UObject
 {
 	GENERATED_BODY()
 public:
-	UClothPhysicalMeshDataBase_Legacy();
-	virtual ~UClothPhysicalMeshDataBase_Legacy();
+	CLOTHINGSYSTEMRUNTIMEINTERFACE_API UClothPhysicalMeshDataBase_Legacy();
+	CLOTHINGSYSTEMRUNTIMEINTERFACE_API virtual ~UClothPhysicalMeshDataBase_Legacy();
 
 	/** Retrieve a registered vertex weight array by unique @param Id. */
-	TArray<float>* GetFloatArray(const uint32 Id) const;
+	CLOTHINGSYSTEMRUNTIMEINTERFACE_API TArray<float>* GetFloatArray(const uint32 Id) const;
 
 	/** Get ids for all registered weight arrays. */
-	TArray<uint32> GetFloatArrayIds() const;
+	CLOTHINGSYSTEMRUNTIMEINTERFACE_API TArray<uint32> GetFloatArrayIds() const;
 
 	/** Get all registered weight arrays. */
-	TArray<TArray<float>*> GetFloatArrays() const;
+	CLOTHINGSYSTEMRUNTIMEINTERFACE_API TArray<TArray<float>*> GetFloatArrays() const;
 
 protected:
 	/** Register an @param Array keyed by a unique @param Id. */
-	void RegisterFloatArray(const uint32 Id, TArray<float> *Array);
+	CLOTHINGSYSTEMRUNTIMEINTERFACE_API void RegisterFloatArray(const uint32 Id, TArray<float> *Array);
 
 public:
 	// Positions of each simulation vertex
diff --git a/Engine/Source/Runtime/ClothingSystemRuntimeInterface/Public/ClothingAssetBase.h b/Engine/Source/Runtime/ClothingSystemRuntimeInterface/Public/ClothingAssetBase.h
--- a/Engine/Source/Runtime/ClothingSystemRuntimeInterface/Public/ClothingAssetBase.h
+++ b/Engine/Source/Runtime/ClothingSystemRuntimeInterface/Public/ClothingAssetBase.h
@@ -12,8 +12,8 @@ class USkeletalMesh;
  * An interface object for any clothing asset the engine can use.
  * Any clothing asset concrete object should derive from this.
  */
-UCLASS(Abstract)
-class CLOTHINGSYSTEMRUNTIMEINTERFACE_API UClothingAssetBase : public UObject
+UCLASS(Abstract, MinimalAPI)
+class UClothingAssetBase : public UObject
 {
 	GENERATED_BODY()
 
@@ -26,14 +26,14 @@ public:
 	* @param InSubmeshIdx Submesh in this asset to replace section with
 	* @param InAssetLodIndex Internal clothing LOD to use
 	*/
-	virtual bool BindToSkeletalMesh(USkeletalMesh* InSkelMesh, const int32 InMeshLodIndex, const int32 InSectionIndex, const int32 InAssetLodIndex)
+	CLOTHINGSYSTEMRUNTIMEINTERFACE_API virtual bool BindToSkeletalMesh(USkeletalMesh* InSkelMesh, const int32 InMeshLodIndex, const int32 InSectionIndex, const int32 InAssetLodIndex)
 	PURE_VIRTUAL(UClothingAssetBase::BindToSkeletalMesh, return false;);
 
 	/**
 	* Unbinds this clothing asset from the provided skeletal mesh, will remove all LODs
 	* @param InSkelMesh skeletal mesh to remove from
 	*/
-	virtual void UnbindFromSkeletalMesh(USkeletalMesh* InSkelMesh)
+	CLOTHINGSYSTEMRUNTIMEINTERFACE_API virtual void UnbindFromSkeletalMesh(USkeletalMesh* InSkelMesh)
 	PURE_VIRTUAL(UClothingAssetBase::UnbindFromSkeletalMesh, );
 
 	/**
@@ -41,14 +41,14 @@ public:
 	* @param InSkelMesh skeletal mesh to remove from
 	* @param InMeshLodIndex Mesh LOD to remove this asset from (could still be bound to other LODs)
 	*/
-	virtual void UnbindFromSkeletalMesh(USkeletalMesh* InSkelMesh, const int32 InMeshLodIndex)
+	CLOTHINGSYSTEMRUNTIMEINTERFACE_API virtual void UnbindFromSkeletalMesh(USkeletalMesh* InSkelMesh, const int32 InMeshLodIndex)
 	PURE_VIRTUAL(UClothingAssetBase::UnbindFromSkeletalMesh,);
 
 	/**
 	 * Update all extra LOD deformer mappings.
 	 * This should be called whenever the raytracing LOD bias is changed.
 	 */
-	virtual void UpdateAllLODBiasMappings(USkeletalMesh* SkeletalMesh)
+	CLOTHINGSYSTEMRUNTIMEINTERFACE_API virtual void UpdateAllLODBiasMappings(USkeletalMesh* SkeletalMesh)
 	PURE_VIRTUAL(UClothingAssetBase::UpdateAllLODBiasMappings,);
 
 	/**
@@ -64,7 +64,7 @@ public:
 	}
 
 	/** Add a new LOD class instance. */
-	virtual int32 AddNewLod()
+	CLOTHINGSYSTEMRUNTIMEINTERFACE_API virtual int32 AddNewLod()
 	PURE_VIRTUAL(UClothingAssetBase::AddNewLod(), return INDEX_NONE;);
 
 	/**
@@ -72,7 +72,7 @@ public:
 	*	When we transition between LODs we skin the incoming mesh to the outgoing mesh
 	*	in exactly the same way the render mesh is skinned to create a smooth swap
 	*/
-	virtual void BuildLodTransitionData()
+	CLOTHINGSYSTEMRUNTIMEINTERFACE_API virtual void BuildLodTransitionData()
 	PURE_VIRTUAL(UClothingAssetBase::BuildLodTransitionData(), );
 #endif
 
@@ -81,7 +81,7 @@ public:
 	 * Called on the clothing asset when the base data (physical mesh, config etc.)
 	 * has changed, so any intermediate generated data can be regenerated.
 	 */
-	virtual void InvalidateAllCachedData()
+	CLOTHINGSYSTEMRUNTIMEINTERFACE_API virtual void InvalidateAllCachedData()
 	PURE_VIRTUAL(UClothingAssetBase::InvalidateCachedData(),);
 #endif // WITH_EDITORONLY_DATA
 
@@ -91,24 +91,24 @@ public:
 	 * data.
 	 * @param InSkelMesh - The mesh to use to remap the bones
 	 */
-	virtual void RefreshBoneMapping(USkeletalMesh* InSkelMesh)
+	CLOTHINGSYSTEMRUNTIMEINTERFACE_API virtual void RefreshBoneMapping(USkeletalMesh* InSkelMesh)
 	PURE_VIRTUAL(UClothingAssetBase::RefreshBoneMapping, );
 
 	/** Check the validity of a LOD index */
-	virtual bool IsValidLod(int32 InLodIndex) const
+	CLOTHINGSYSTEMRUNTIMEINTERFACE_API virtual bool IsValidLod(int32 InLodIndex) const
 	PURE_VIRTUAL(UClothingAssetBase::IsValidLod(), return false;);
 
 	/** Get the number of LODs defined in the clothing asset */
-	virtual int32 GetNumLods() const
+	CLOTHINGSYSTEMRUNTIMEINTERFACE_API virtual int32 GetNumLods() const
 	PURE_VIRTUAL(UClothingAssetBase::GetNumLods(), return 0;);
 
 	/** Builds self collision data */
 	UE_DEPRECATED(5.0, "Cached data are now all rebuilt by calling InvalidateCachedData()")
-	virtual void BuildSelfCollisionData()
+	CLOTHINGSYSTEMRUNTIMEINTERFACE_API virtual void BuildSelfCollisionData()
 	PURE_VIRTUAL(UClothingAssetBase::BuildSelfCollisionData(), );
 
 	/** Called after all cloth assets sharing the same simulation are added or loaded */
-	virtual void PostUpdateAllAssets()
+	CLOTHINGSYSTEMRUNTIMEINTERFACE_API virtual void PostUpdateAllAssets()
 	PURE_VIRTUAL(UClothingAssetBase::PostUpdateAllAssets(), );
 
 	/** Get the guid identifying this asset */
diff --git a/Engine/Source/Runtime/ClothingSystemRuntimeInterface/Public/ClothingSimulationFactory.h b/Engine/Source/Runtime/ClothingSystemRuntimeInterface/Public/ClothingSimulationFactory.h
--- a/Engine/Source/Runtime/ClothingSystemRuntimeInterface/Public/ClothingSimulationFactory.h
+++ b/Engine/Source/Runtime/ClothingSystemRuntimeInterface/Public/ClothingSimulationFactory.h
@@ -24,12 +24,12 @@ class UEnum;
 
 // An interface for a class that will provide default simulation factory classes
 // Used by modules wanting to override clothing simulation to provide their own implementation
-class CLOTHINGSYSTEMRUNTIMEINTERFACE_API IClothingSimulationFactoryClassProvider : public IModularFeature
+class IClothingSimulationFactoryClassProvider : public IModularFeature
 {
 public:
 
 	// The feature name to register against for providers
-	static const FName FeatureName;
+	static CLOTHINGSYSTEMRUNTIMEINTERFACE_API const FName FeatureName;
 
 	// Called by the engine to get the default clothing simulation factory to use
 	// for skeletal mesh components (see USkeletalMeshComponent constructor).
@@ -44,8 +44,8 @@ public:
 };
 
 // Any clothing simulation factory should derive from this interface object to interact with the engine
-UCLASS(Abstract)
-class CLOTHINGSYSTEMRUNTIMEINTERFACE_API UClothingSimulationFactory : public UObject
+UCLASS(Abstract, MinimalAPI)
+class UClothingSimulationFactory : public UObject
 {
 	GENERATED_BODY()
 
@@ -53,36 +53,36 @@ public:
 	// Return the default clothing simulation factory class as set by the build or by
 	// the p.Cloth.DefaultClothingSimulationFactoryClass console variable if any available.
 	// Otherwise return the last registered factory.
-	static TSubclassOf<class UClothingSimulationFactory> GetDefaultClothingSimulationFactoryClass();
+	static CLOTHINGSYSTEMRUNTIMEINTERFACE_API TSubclassOf<class UClothingSimulationFactory> GetDefaultClothingSimulationFactoryClass();
 
 	// Create a simulation object for a skeletal mesh to use (see IClothingSimulation)
-	virtual IClothingSimulation* CreateSimulation()
+	CLOTHINGSYSTEMRUNTIMEINTERFACE_API virtual IClothingSimulation* CreateSimulation()
 	PURE_VIRTUAL(UClothingSimulationFactory::CreateSimulation, return nullptr;);
 
 	// Destroy a simulation object, guaranteed to be a pointer returned from CreateSimulation for this factory
-	virtual void DestroySimulation(IClothingSimulation* InSimulation)
+	CLOTHINGSYSTEMRUNTIMEINTERFACE_API virtual void DestroySimulation(IClothingSimulation* InSimulation)
 	PURE_VIRTUAL(UClothingSimulationFactory::DestroySimulation, );
 
 	// Given an asset, decide whether this factory can create a simulation to use the data inside
 	// (return false if data is invalid or missing in the case of custom data)
-	virtual bool SupportsAsset(UClothingAssetBase* InAsset)
+	CLOTHINGSYSTEMRUNTIMEINTERFACE_API virtual bool SupportsAsset(UClothingAssetBase* InAsset)
 	PURE_VIRTUAL(UClothingSimulationFactory::SupportsAsset, return false;);
 
 	// Whether or not we provide an interactor object to manipulate the simulation at runtime.
 	// If true is returned then CreateInteractor *must* create a valid object to handle this
-	virtual bool SupportsRuntimeInteraction()
+	CLOTHINGSYSTEMRUNTIMEINTERFACE_API virtual bool SupportsRuntimeInteraction()
 	PURE_VIRTUAL(UClothingSimulationFactory::SupportsRuntimeInteraction, return false;);
 
 	// Creates the runtime interactor object for a clothing simulation. This object will
 	// receive events allowing it to write data to the simulation context in a safe manner
-	virtual UClothingSimulationInteractor* CreateInteractor()
+	CLOTHINGSYSTEMRUNTIMEINTERFACE_API virtual UClothingSimulationInteractor* CreateInteractor()
 	PURE_VIRTUAL(UClothingSimulationFactory::CreateInteractor, return nullptr;);
 
 	// Return the cloth config type for this cloth factory
-	virtual TArrayView<const TSubclassOf<UClothConfigBase>> GetClothConfigClasses() const
+	CLOTHINGSYSTEMRUNTIMEINTERFACE_API virtual TArrayView<const TSubclassOf<UClothConfigBase>> GetClothConfigClasses() const
 	PURE_VIRTUAL(UClothingSimulationFactory::GetClothConfigClasses, return TArrayView<const TSubclassOf<UClothConfigBase>>(););
 
 	// Return an enum of the weight map targets that can be used with this simulation.
-	virtual const UEnum* GetWeightMapTargetEnum() const
+	CLOTHINGSYSTEMRUNTIMEINTERFACE_API virtual const UEnum* GetWeightMapTargetEnum() const
 	PURE_VIRTUAL(UClothingSimulationFactory::GetWeightMapTargetEnum, return nullptr;);
 };
diff --git a/Engine/Source/Runtime/ClothingSystemRuntimeInterface/Public/ClothingSimulationInteractor.h b/Engine/Source/Runtime/ClothingSystemRuntimeInterface/Public/ClothingSimulationInteractor.h
--- a/Engine/Source/Runtime/ClothingSystemRuntimeInterface/Public/ClothingSimulationInteractor.h
+++ b/Engine/Source/Runtime/ClothingSystemRuntimeInterface/Public/ClothingSimulationInteractor.h
@@ -24,8 +24,8 @@ template <typename T> struct TObjectPtr;
  * Abstract class to control clothing specific interaction.
  * Must be cast to the end used clothing simulation object before use.
  */
-UCLASS(Abstract, BlueprintType)
-class CLOTHINGSYSTEMRUNTIMEINTERFACE_API UClothingInteractor : public UObject
+UCLASS(Abstract, BlueprintType, MinimalAPI)
+class UClothingInteractor : public UObject
 {
 	GENERATED_BODY()
 
@@ -45,17 +45,17 @@ protected:
  * Only write to the simulation and context during the call to Sync, as that is
  * guaranteed to be a safe place to access this data.
  */
-UCLASS(Abstract, BlueprintType)
-class CLOTHINGSYSTEMRUNTIMEINTERFACE_API UClothingSimulationInteractor : public UObject
+UCLASS(Abstract, BlueprintType, MinimalAPI)
+class UClothingSimulationInteractor : public UObject
 {
 	GENERATED_BODY()
 
 public:
 	/** Create a cloth specific interactor. */
-	void CreateClothingInteractor(const UClothingAssetBase* ClothingAsset, int32 ClothingId);
+	CLOTHINGSYSTEMRUNTIMEINTERFACE_API void CreateClothingInteractor(const UClothingAssetBase* ClothingAsset, int32 ClothingId);
 
 	/** Destroy all interactors. */
-	void DestroyClothingInteractors();
+	CLOTHINGSYSTEMRUNTIMEINTERFACE_API void DestroyClothingInteractors();
 
 	// Basic interface that clothing simulations are required to support
 
@@ -66,46 +66,46 @@ public:
 	 * written in a way to pick up these changes on the next update.
 	 * Any inherited class must call this function to also Sync the ClothingInteractors.
 	 */
-	virtual void Sync(IClothingSimulation* Simulation, IClothingSimulationContext* Context);
+	CLOTHINGSYSTEMRUNTIMEINTERFACE_API virtual void Sync(IClothingSimulation* Simulation, IClothingSimulationContext* Context);
 
 	/** Called to update collision status without restarting the simulation. */
 	UFUNCTION(BlueprintCallable, Category=ClothingSimulation)
-	virtual void PhysicsAssetUpdated()
+	CLOTHINGSYSTEMRUNTIMEINTERFACE_API virtual void PhysicsAssetUpdated()
 	PURE_VIRTUAL(UClothingSimulationInteractor::PhysicsAssetUpdated, );
 
 	/** Called to update the cloth config without restarting the simulation. */
 	UFUNCTION(BlueprintCallable, Category=ClothingSimulation)
-	virtual void ClothConfigUpdated()
+	CLOTHINGSYSTEMRUNTIMEINTERFACE_API virtual void ClothConfigUpdated()
 	PURE_VIRTUAL(UClothingSimulationInteractor::ClothConfigUpdated, );
 
 	/** Set the stiffness of the spring force for the animation drive. */
 	UFUNCTION(BlueprintCallable, Category = ClothingSimulation)
-	virtual void SetAnimDriveSpringStiffness(float InStiffness)
+	CLOTHINGSYSTEMRUNTIMEINTERFACE_API virtual void SetAnimDriveSpringStiffness(float InStiffness)
 	PURE_VIRTUAL(UClothingSimulationInteractor::SetAnimDriveSpringStiffness, );
 
 	/** Set a new gravity override and enable the override. */
 	UFUNCTION(BlueprintCallable, Category = ClothingSimulation)
-	virtual void EnableGravityOverride(const FVector& InVector)
+	CLOTHINGSYSTEMRUNTIMEINTERFACE_API virtual void EnableGravityOverride(const FVector& InVector)
 	PURE_VIRTUAL(UClothingSimulationInteractor::EnableGravityOverride, );
 
 	/** Disable any currently set gravity override. */
 	UFUNCTION(BlueprintCallable, Category = ClothingSimulation)
-	virtual void DisableGravityOverride()
+	CLOTHINGSYSTEMRUNTIMEINTERFACE_API virtual void DisableGravityOverride()
 	PURE_VIRTUAL(UClothingSimulationInteractor::DisableGravityOverride, );
 
 	/** Set the number of time dependent solver iterations. */
 	UFUNCTION(BlueprintCallable, Category = ClothingSimulation)
-	virtual void SetNumIterations(int32 NumIterations = 1)
+	CLOTHINGSYSTEMRUNTIMEINTERFACE_API virtual void SetNumIterations(int32 NumIterations = 1)
 	PURE_VIRTUAL(UClothingSimulationInteractor::SetNumIterations, );
 
 	/** Set the maximum number of solver iterations. */
 	UFUNCTION(BlueprintCallable, Category = ClothingSimulation)
-	virtual void SetMaxNumIterations(int32 MaxNumIterations = 10)
+	CLOTHINGSYSTEMRUNTIMEINTERFACE_API virtual void SetMaxNumIterations(int32 MaxNumIterations = 10)
 	PURE_VIRTUAL(UClothingSimulationInteractor::SetMaxNumIterations, );
 
 	/** Set the number of substeps or subdivisions. */
 	UFUNCTION(BlueprintCallable, Category = ClothingSimulation, Meta = (Keywords = "Subdivisions"))
-	virtual void SetNumSubsteps(int32 NumSubsteps = 1)
+	CLOTHINGSYSTEMRUNTIMEINTERFACE_API virtual void SetNumSubsteps(int32 NumSubsteps = 1)
 	PURE_VIRTUAL(UClothingSimulationInteractor::SetNumSubsteps, );
 
 	// Base clothing simulations interface
@@ -142,14 +142,14 @@ public:
 
 	/**Return a cloth interactor for this simulation. */
 	UFUNCTION(BlueprintCallable, Category = ClothingSimulation)
-	UClothingInteractor* GetClothingInteractor(const FString& ClothingAssetName) const;
+	CLOTHINGSYSTEMRUNTIMEINTERFACE_API UClothingInteractor* GetClothingInteractor(const FString& ClothingAssetName) const;
 
 	/** Cloth interactors currently created. */
 	UPROPERTY()
 	TMap<FName, TObjectPtr<UClothingInteractor>> ClothingInteractors;
 
 protected:
-	virtual UClothingInteractor* CreateClothingInteractor()
+	CLOTHINGSYSTEMRUNTIMEINTERFACE_API virtual UClothingInteractor* CreateClothingInteractor()
 	PURE_VIRTUAL(UClothingSimulationInteractor::CreateClothingInteractor, return nullptr;);
 
 private:
diff --git a/Engine/Source/Runtime/ClothingSystemRuntimeInterface/Public/ClothingSimulationInterface.h b/Engine/Source/Runtime/ClothingSystemRuntimeInterface/Public/ClothingSimulationInterface.h
--- a/Engine/Source/Runtime/ClothingSystemRuntimeInterface/Public/ClothingSimulationInterface.h
+++ b/Engine/Source/Runtime/ClothingSystemRuntimeInterface/Public/ClothingSimulationInterface.h
@@ -13,19 +13,19 @@ struct FClothCollisionData;
 struct FClothSimulData;
 
 /** Empty interface, derived simulation modules define the contents of the context. */
-class CLOTHINGSYSTEMRUNTIMEINTERFACE_API IClothingSimulationContext
+class IClothingSimulationContext
 {
 public:
-	IClothingSimulationContext();
-	virtual ~IClothingSimulationContext();
+	CLOTHINGSYSTEMRUNTIMEINTERFACE_API IClothingSimulationContext();
+	CLOTHINGSYSTEMRUNTIMEINTERFACE_API virtual ~IClothingSimulationContext();
 };
 
 /** Base class for clothing simulators. */
-class CLOTHINGSYSTEMRUNTIMEINTERFACE_API IClothingSimulation
+class IClothingSimulation
 {
 public:
-	IClothingSimulation();
-	virtual ~IClothingSimulation();
+	CLOTHINGSYSTEMRUNTIMEINTERFACE_API IClothingSimulation();
+	CLOTHINGSYSTEMRUNTIMEINTERFACE_API virtual ~IClothingSimulation();
 
 	// The majority of the API for this class is protected. The required objects (skel meshes and the parallel task)
 	// are friends so they can use the functionality. For the most part the simulation is not designed to be used
diff --git a/Engine/Source/Runtime/ClothingSystemRuntimeNv/Public/ClothConfigNv.h b/Engine/Source/Runtime/ClothingSystemRuntimeNv/Public/ClothConfigNv.h
--- a/Engine/Source/Runtime/ClothingSystemRuntimeNv/Public/ClothConfigNv.h
+++ b/Engine/Source/Runtime/ClothingSystemRuntimeNv/Public/ClothConfigNv.h
@@ -61,24 +61,24 @@ enum class EClothingWindMethodNv : uint8
 
 /** Holds initial, asset level config for clothing actors. */
 PRAGMA_DISABLE_DEPRECATION_WARNINGS
-UCLASS()
-class CLOTHINGSYSTEMRUNTIMENV_API UClothConfigNv : public UClothConfigCommon
+UCLASS(MinimalAPI)
+class UClothConfigNv : public UClothConfigCommon
 {
 	GENERATED_BODY()
 public:
-	UClothConfigNv();
+	CLOTHINGSYSTEMRUNTIMENV_API UClothConfigNv();
 
 	// Set up custom version serialization.
-	virtual void Serialize(FArchive& Ar) override;
+	CLOTHINGSYSTEMRUNTIMENV_API virtual void Serialize(FArchive& Ar) override;
 
 	// Update the deprecated properties.
-	virtual void PostLoad() override;
+	CLOTHINGSYSTEMRUNTIMENV_API virtual void PostLoad() override;
 
 	// Migrate from the legacy FClothConfig structure.
-	virtual void MigrateFrom(const FClothConfig_Legacy& ClothConfig) override;
+	CLOTHINGSYSTEMRUNTIMENV_API virtual void MigrateFrom(const FClothConfig_Legacy& ClothConfig) override;
 
 	// Migrate to the legacy FClothConfig structure.
-	virtual bool MigrateTo(FClothConfig_Legacy& ClothConfig) const override;
+	CLOTHINGSYSTEMRUNTIMENV_API virtual bool MigrateTo(FClothConfig_Legacy& ClothConfig) const override;
 
 	/** Return whether to pre-compute self collision data. */
 	virtual bool NeedsSelfCollisionData() const override { return SelfCollisionRadius > 0.0f && SelfCollisionStiffness > 0.0f; }
diff --git a/Engine/Source/Runtime/ClothingSystemRuntimeNv/Public/ClothPhysicalMeshDataNv_Legacy.h b/Engine/Source/Runtime/ClothingSystemRuntimeNv/Public/ClothPhysicalMeshDataNv_Legacy.h
--- a/Engine/Source/Runtime/ClothingSystemRuntimeNv/Public/ClothPhysicalMeshDataNv_Legacy.h
+++ b/Engine/Source/Runtime/ClothingSystemRuntimeNv/Public/ClothPhysicalMeshDataNv_Legacy.h
@@ -14,13 +14,13 @@ class UObject;
  * Deprecated, use FClothPhysicalMeshData instead.
  * NV specific spatial simulation data for a mesh.
  */
-UCLASS()
-class CLOTHINGSYSTEMRUNTIMENV_API UClothPhysicalMeshDataNv_Legacy: public UClothPhysicalMeshDataBase_Legacy
+UCLASS(MinimalAPI)
+class UClothPhysicalMeshDataNv_Legacy: public UClothPhysicalMeshDataBase_Legacy
 {
 	GENERATED_BODY()
 public:
-	UClothPhysicalMeshDataNv_Legacy();
-	virtual ~UClothPhysicalMeshDataNv_Legacy();
+	CLOTHINGSYSTEMRUNTIMENV_API UClothPhysicalMeshDataNv_Legacy();
+	CLOTHINGSYSTEMRUNTIMENV_API virtual ~UClothPhysicalMeshDataNv_Legacy();
 
 	// The distance that each vertex can move away from its reference (skinned) position
 	UPROPERTY()
diff --git a/Engine/Source/Runtime/ClothingSystemRuntimeNv/Public/ClothingSimulationFactoryNv.h b/Engine/Source/Runtime/ClothingSystemRuntimeNv/Public/ClothingSimulationFactoryNv.h
--- a/Engine/Source/Runtime/ClothingSystemRuntimeNv/Public/ClothingSimulationFactoryNv.h
+++ b/Engine/Source/Runtime/ClothingSystemRuntimeNv/Public/ClothingSimulationFactoryNv.h
@@ -20,19 +20,19 @@ template <class TClass> class TSubclassOf;
 
 class UE_DEPRECATED(5.1, "NvCloth simulation is no longer supported, UChaosClothingSimulationFactory should be used going forward.") UClothingSimulationFactoryNv;
 
-UCLASS()
-class CLOTHINGSYSTEMRUNTIMENV_API UClothingSimulationFactoryNv final : public UClothingSimulationFactory
+UCLASS(MinimalAPI)
+class UClothingSimulationFactoryNv final : public UClothingSimulationFactory
 {
 	GENERATED_BODY()
 public:
 
-	virtual IClothingSimulation* CreateSimulation() override;
-	virtual void DestroySimulation(IClothingSimulation* InSimulation) override;
-	virtual bool SupportsAsset(UClothingAssetBase* InAsset) override;
+	CLOTHINGSYSTEMRUNTIMENV_API virtual IClothingSimulation* CreateSimulation() override;
+	CLOTHINGSYSTEMRUNTIMENV_API virtual void DestroySimulation(IClothingSimulation* InSimulation) override;
+	CLOTHINGSYSTEMRUNTIMENV_API virtual bool SupportsAsset(UClothingAssetBase* InAsset) override;
 
-	virtual bool SupportsRuntimeInteraction() override;
-	virtual UClothingSimulationInteractor* CreateInteractor() override;
+	CLOTHINGSYSTEMRUNTIMENV_API virtual bool SupportsRuntimeInteraction() override;
+	CLOTHINGSYSTEMRUNTIMENV_API virtual UClothingSimulationInteractor* CreateInteractor() override;
 
-	virtual TArrayView<const TSubclassOf<UClothConfigBase>> GetClothConfigClasses() const override;
-	virtual const UEnum* GetWeightMapTargetEnum() const override;
+	CLOTHINGSYSTEMRUNTIMENV_API virtual TArrayView<const TSubclassOf<UClothConfigBase>> GetClothConfigClasses() const override;
+	CLOTHINGSYSTEMRUNTIMENV_API virtual const UEnum* GetWeightMapTargetEnum() const override;
 };
diff --git a/Engine/Source/Runtime/ClothingSystemRuntimeNv/Public/ClothingSimulationInteractorNv.h b/Engine/Source/Runtime/ClothingSystemRuntimeNv/Public/ClothingSimulationInteractorNv.h
--- a/Engine/Source/Runtime/ClothingSystemRuntimeNv/Public/ClothingSimulationInteractorNv.h
+++ b/Engine/Source/Runtime/ClothingSystemRuntimeNv/Public/ClothingSimulationInteractorNv.h
@@ -23,21 +23,21 @@ struct FFrame;
 DECLARE_DELEGATE_TwoParams(NvInteractorCommand, FClothingSimulationNv*, FClothingSimulationContextNv*)
 
 class UE_DEPRECATED(5.1, "NvCloth is no longer a supported clothing simlation provider, prefer the Chaos simulation moving forward.") UClothingSimulationInteractorNv;
-UCLASS(BlueprintType)
-class CLOTHINGSYSTEMRUNTIMENV_API UClothingSimulationInteractorNv : public UClothingSimulationInteractor
+UCLASS(BlueprintType, MinimalAPI)
+class UClothingSimulationInteractorNv : public UClothingSimulationInteractor
 {
 	GENERATED_BODY()
 
 public:
 
 	// UClothingSimulationInteractor Interface
-	virtual void Sync(IClothingSimulation* InSimulation, IClothingSimulationContext* InContext) override;
+	CLOTHINGSYSTEMRUNTIMENV_API virtual void Sync(IClothingSimulation* InSimulation, IClothingSimulationContext* InContext) override;
 
-	virtual void PhysicsAssetUpdated() override;
-	virtual void ClothConfigUpdated() override;
-	virtual void SetAnimDriveSpringStiffness(float InStiffness) override;
-	virtual void EnableGravityOverride(const FVector& InVector) override;
-	virtual void DisableGravityOverride() override;
+	CLOTHINGSYSTEMRUNTIMENV_API virtual void PhysicsAssetUpdated() override;
+	CLOTHINGSYSTEMRUNTIMENV_API virtual void ClothConfigUpdated() override;
+	CLOTHINGSYSTEMRUNTIMENV_API virtual void SetAnimDriveSpringStiffness(float InStiffness) override;
+	CLOTHINGSYSTEMRUNTIMENV_API virtual void EnableGravityOverride(const FVector& InVector) override;
+	CLOTHINGSYSTEMRUNTIMENV_API virtual void DisableGravityOverride() override;
 
 	// TODO: These new functions are currently unimplemented
 	virtual void SetNumIterations(int32 /*NumIterations*/) override {}
@@ -47,7 +47,7 @@ public:
 
 	// Set the stiffness of the resistive damping force for anim drive
 	UFUNCTION(BlueprintCallable, Category=ClothingSimulation)
-	void SetAnimDriveDamperStiffness(float InStiffness);
+	CLOTHINGSYSTEMRUNTIMENV_API void SetAnimDriveDamperStiffness(float InStiffness);
 
 protected:
 
diff --git a/Engine/Source/Runtime/ColorManagement/Public/ColorSpace.h b/Engine/Source/Runtime/ColorManagement/Public/ColorSpace.h
--- a/Engine/Source/Runtime/ColorManagement/Public/ColorSpace.h
+++ b/Engine/Source/Runtime/ColorManagement/Public/ColorSpace.h
@@ -126,7 +126,7 @@ inline FVector2d GetWhitePoint(EWhitePoint InWhitePoint)
 
 
 /** Color space definition as 4 chromaticity coordinates, in double precision internally. */
-class COLORMANAGEMENT_API FColorSpace
+class FColorSpace
 {
 public:
 	/**
@@ -134,14 +134,14 @@ public:
 	 * 
 	 * @return FColorSpace working color space
 	 */
-	static const FColorSpace& GetWorking();
+	static COLORMANAGEMENT_API const FColorSpace& GetWorking();
 
 	/**
 	 * Set the global engine working color space (as a singleton).
 	 * 
 	 * @param ColorSpace working color space
 	 */
-	static void SetWorking(FColorSpace ColorSpace);
+	static COLORMANAGEMENT_API void SetWorking(FColorSpace ColorSpace);
 
 
 	/** Constructor */
@@ -155,14 +155,14 @@ public:
 	* @param InBlue Chromaticity 2D coordinates for the blue color.
 	* @param InWhite Chromaticity 2D coordinates for the white point.
 	*/
-	explicit FColorSpace(const FVector2d& InRed, const FVector2d& InGreen, const FVector2d& InBlue, const FVector2d& InWhite);
+	COLORMANAGEMENT_API explicit FColorSpace(const FVector2d& InRed, const FVector2d& InGreen, const FVector2d& InBlue, const FVector2d& InWhite);
 
 	/**
 	* Constructor
 	*
 	* @param ColorSpaceType Color space type.
 	*/
-	explicit FColorSpace(UE::Color::EColorSpace ColorSpaceType);
+	COLORMANAGEMENT_API explicit FColorSpace(UE::Color::EColorSpace ColorSpaceType);
 
 	FColorSpace(FColorSpace&&) = default;
 	FColorSpace(const FColorSpace&) = default;
@@ -275,7 +275,7 @@ public:
 	 * @param Tolerance Error tolerance.
 	 * @return true if the color spaces are equal within tolerance limits, false otherwise.
 	 */
-	bool Equals(const FColorSpace& ColorSpace, double Tolerance = 1.e-7) const;
+	COLORMANAGEMENT_API bool Equals(const FColorSpace& ColorSpace, double Tolerance = 1.e-7) const;
 
 
 	/**
@@ -285,19 +285,19 @@ public:
 	 * @param Tolerance Error tolerance.
 	 * @return true if the color spaces are equal within tolerance limits, false otherwise.
 	 */
-	bool Equals(const TStaticArray<FVector2d, 4>& InChromaticities, double Tolerance = 1.e-7) const;
+	COLORMANAGEMENT_API bool Equals(const TStaticArray<FVector2d, 4>& InChromaticities, double Tolerance = 1.e-7) const;
 
 	/**
 	 * Convenience function to verify if the color space matches the engine's default sRGB chromaticities.
 	 *
 	 * @return true if sRGB.
 	 */
-	bool IsSRGB() const;
+	COLORMANAGEMENT_API bool IsSRGB() const;
 
 	/**
 	* Converts temperature in Kelvins of a black body radiator to an RGB color in the current space.
 	*/
-	FLinearColor MakeFromColorTemperature(float Temp) const;
+	COLORMANAGEMENT_API FLinearColor MakeFromColorTemperature(float Temp) const;
 
 	/**
 	 * Calculate the color's luminance value in the current space.
@@ -305,18 +305,18 @@ public:
 	 * @param Color The sampled color.
 	 * @return float Luminance value.
 	 */
-	float GetLuminance(const FLinearColor& Color) const;
+	COLORMANAGEMENT_API float GetLuminance(const FLinearColor& Color) const;
 
 	/**
 	 * Get the luminance factors in the current space.
 	 *
 	 * @return FLinearColor Luminance factors.
 	 */
-	FLinearColor GetLuminanceFactors() const;
+	COLORMANAGEMENT_API FLinearColor GetLuminanceFactors() const;
 
 private:
 
-	FMatrix44d CalcRgbToXYZ() const;
+	COLORMANAGEMENT_API FMatrix44d CalcRgbToXYZ() const;
 
 	/** Red, green, blue, white chromaticities, in order. */
 	TStaticArray<FVector2d, 4> Chromaticities;
diff --git a/Engine/Source/Runtime/CookOnTheFly/Internal/CookOnTheFly.h b/Engine/Source/Runtime/CookOnTheFly/Internal/CookOnTheFly.h
--- a/Engine/Source/Runtime/CookOnTheFly/Internal/CookOnTheFly.h
+++ b/Engine/Source/Runtime/CookOnTheFly/Internal/CookOnTheFly.h
@@ -138,7 +138,7 @@ class FCookOnTheFlyMessage
 {
 public:
 	/** Creates a new instance of a cook-on-the-fly message. */
-	COOKONTHEFLY_API FCookOnTheFlyMessage() = default;
+	FCookOnTheFlyMessage() = default;
 
 	/** Creates a new instance of a cook-on-the-fly message with the specified message type. */
 	explicit FCookOnTheFlyMessage(ECookOnTheFlyMessage MessageType)
@@ -251,7 +251,7 @@ class FCookOnTheFlyRequest
 {
 public:
 	/** Creates a new instance of a cook-on-the-fly request. */
-	COOKONTHEFLY_API FCookOnTheFlyRequest() = default;
+	FCookOnTheFlyRequest() = default;
 
 	/** Creates a new instance of a cook-on-the-fly request with the specified request type. */
 	explicit FCookOnTheFlyRequest(ECookOnTheFlyMessage MessageType)
@@ -268,7 +268,7 @@ class FCookOnTheFlyResponse
 {
 public:
 	/** Creates a new instance of a cook-on-the-fly response. */
-	COOKONTHEFLY_API FCookOnTheFlyResponse() = default;
+	FCookOnTheFlyResponse() = default;
 
 	/** Creates a new instance of a cook-on-the-fly response with the specified response type. */
 	explicit FCookOnTheFlyResponse(ECookOnTheFlyMessage MessageType)
diff --git a/Engine/Source/Runtime/DeveloperSettings/Public/Engine/DeveloperSettings.h b/Engine/Source/Runtime/DeveloperSettings/Public/Engine/DeveloperSettings.h
--- a/Engine/Source/Runtime/DeveloperSettings/Public/Engine/DeveloperSettings.h
+++ b/Engine/Source/Runtime/DeveloperSettings/Public/Engine/DeveloperSettings.h
@@ -19,40 +19,40 @@ class SWidget;
 /**
  * The base class of any auto discovered settings object.
  */
-UCLASS(Abstract)
-class DEVELOPERSETTINGS_API UDeveloperSettings : public UObject
+UCLASS(Abstract, MinimalAPI)
+class UDeveloperSettings : public UObject
 {
 	GENERATED_BODY()
 
 public:
-	UDeveloperSettings(const FObjectInitializer& ObjectInitializer = FObjectInitializer::Get());
+	DEVELOPERSETTINGS_API UDeveloperSettings(const FObjectInitializer& ObjectInitializer = FObjectInitializer::Get());
 
 	/** Gets the settings container name for the settings, either Project or Editor */
-	virtual FName GetContainerName() const;
+	DEVELOPERSETTINGS_API virtual FName GetContainerName() const;
 	/** Gets the category for the settings, some high level grouping like, Editor, Engine, Game...etc. */
-	virtual FName GetCategoryName() const;
+	DEVELOPERSETTINGS_API virtual FName GetCategoryName() const;
 	/** The unique name for your section of settings, uses the class's FName. */
-	virtual FName GetSectionName() const;
+	DEVELOPERSETTINGS_API virtual FName GetSectionName() const;
 
 #if WITH_EDITOR
 	/** Gets the section text, uses the classes DisplayName by default. */
-	virtual FText GetSectionText() const;
+	DEVELOPERSETTINGS_API virtual FText GetSectionText() const;
 	/** Gets the description for the section, uses the classes ToolTip by default. */
-	virtual FText GetSectionDescription() const;
+	DEVELOPERSETTINGS_API virtual FText GetSectionDescription() const;
 
 	/** Whether or not this class supports auto registration or if the settings have a custom setup */
 	virtual bool SupportsAutoRegistration() const { return true; }
 
 	/** Returns a delegate that can be used to monitor for property changes to this object */
 	DECLARE_MULTICAST_DELEGATE_TwoParams(FOnSettingsChanged, UObject*, struct FPropertyChangedEvent&);
-	FOnSettingsChanged& OnSettingChanged();
+	DEVELOPERSETTINGS_API FOnSettingsChanged& OnSettingChanged();
 
 	/** UObject interface */
-	virtual void PostEditChangeProperty(struct FPropertyChangedEvent& PropertyChangedEvent) override;
+	DEVELOPERSETTINGS_API virtual void PostEditChangeProperty(struct FPropertyChangedEvent& PropertyChangedEvent) override;
 #endif
 
 	/** Gets a custom widget for the settings.  This is only for very custom situations. */
-	virtual TSharedPtr<SWidget> GetCustomSettingsWidget() const;
+	DEVELOPERSETTINGS_API virtual TSharedPtr<SWidget> GetCustomSettingsWidget() const;
 
 protected:
 	/**
@@ -67,9 +67,9 @@ protected:
 	
 #if WITH_EDITOR
 	/** Populates all properties that have 'ConsoleVariable' meta data with the respective console variable values */
-	void ImportConsoleVariableValues();
+	DEVELOPERSETTINGS_API void ImportConsoleVariableValues();
 	/** If property has 'ConsoleVariable' meta data, exports the property value to the specified console variable */
-	void ExportValuesToConsoleVariables(FProperty* PropertyThatChanged);
+	DEVELOPERSETTINGS_API void ExportValuesToConsoleVariables(FProperty* PropertyThatChanged);
 
 	/** Holds a delegate that is executed after the settings section has been modified. */
 	FOnSettingsChanged SettingsChangedDelegate;
diff --git a/Engine/Source/Runtime/DeveloperSettings/Public/Engine/DeveloperSettingsBackedByCVars.h b/Engine/Source/Runtime/DeveloperSettings/Public/Engine/DeveloperSettingsBackedByCVars.h
--- a/Engine/Source/Runtime/DeveloperSettings/Public/Engine/DeveloperSettingsBackedByCVars.h
+++ b/Engine/Source/Runtime/DeveloperSettings/Public/Engine/DeveloperSettingsBackedByCVars.h
@@ -16,17 +16,17 @@ class UObject;
  * The base class of auto discovered settings object where some or all of the settings
  * are stored in console variables instead of config variables.
  */
-UCLASS(Abstract)
-class DEVELOPERSETTINGS_API UDeveloperSettingsBackedByCVars : public UDeveloperSettings
+UCLASS(Abstract, MinimalAPI)
+class UDeveloperSettingsBackedByCVars : public UDeveloperSettings
 {
 	GENERATED_BODY()
 
 public:
-	UDeveloperSettingsBackedByCVars(const FObjectInitializer& ObjectInitializer = FObjectInitializer::Get());
+	DEVELOPERSETTINGS_API UDeveloperSettingsBackedByCVars(const FObjectInitializer& ObjectInitializer = FObjectInitializer::Get());
 
-	virtual void PostInitProperties() override;
+	DEVELOPERSETTINGS_API virtual void PostInitProperties() override;
 #if WITH_EDITOR
-	virtual void PostEditChangeProperty(struct FPropertyChangedEvent& PropertyChangedEvent) override;
+	DEVELOPERSETTINGS_API virtual void PostEditChangeProperty(struct FPropertyChangedEvent& PropertyChangedEvent) override;
 #endif
 
 };
diff --git a/Engine/Source/Runtime/DeveloperSettings/Public/Engine/PlatformSettings.h b/Engine/Source/Runtime/DeveloperSettings/Public/Engine/PlatformSettings.h
--- a/Engine/Source/Runtime/DeveloperSettings/Public/Engine/PlatformSettings.h
+++ b/Engine/Source/Runtime/DeveloperSettings/Public/Engine/PlatformSettings.h
@@ -21,11 +21,11 @@ class UPlatformSettings;
 class UPlatformSettingsManager;
 
 USTRUCT()
-struct DEVELOPERSETTINGS_API FPerPlatformSettings
+struct FPerPlatformSettings
 {
 	GENERATED_BODY()
 
-	void Initialize(TSubclassOf<UPlatformSettings> SettingsClass);
+	DEVELOPERSETTINGS_API void Initialize(TSubclassOf<UPlatformSettings> SettingsClass);
 
 private:
 	UPROPERTY(Instanced, Transient, EditAnywhere, EditFixedSize, Category = Layout)
@@ -54,13 +54,13 @@ private:
  *         UMyPerPlatformSettings* MySettings = UPlatformSettingsManager::Get().GetSettingsForPlatform<UMyPerPlatformSettings>()
  *         that will get you the current settings for the active platform, or the simulated platform in the editor.
  */
-UCLASS(Abstract, perObjectConfig)
-class DEVELOPERSETTINGS_API UPlatformSettings : public UObject
+UCLASS(Abstract, perObjectConfig, MinimalAPI)
+class UPlatformSettings : public UObject
 {
 	GENERATED_BODY()
 
 public:
-	UPlatformSettings(const FObjectInitializer& ObjectInitializer = FObjectInitializer::Get());
+	DEVELOPERSETTINGS_API UPlatformSettings(const FObjectInitializer& ObjectInitializer = FObjectInitializer::Get());
 
 	virtual void InitializePlatformDefaults() { }
 
diff --git a/Engine/Source/Runtime/DeveloperSettings/Public/Engine/PlatformSettingsManager.h b/Engine/Source/Runtime/DeveloperSettings/Public/Engine/PlatformSettingsManager.h
--- a/Engine/Source/Runtime/DeveloperSettings/Public/Engine/PlatformSettingsManager.h
+++ b/Engine/Source/Runtime/DeveloperSettings/Public/Engine/PlatformSettingsManager.h
@@ -36,13 +36,13 @@ struct FPlatformSettingsInstances
 };
 
 // The manager for all platform-specific settings
-UCLASS()
-class DEVELOPERSETTINGS_API UPlatformSettingsManager : public UObject
+UCLASS(MinimalAPI)
+class UPlatformSettingsManager : public UObject
 {
 	GENERATED_BODY()
 
 public:
-	UPlatformSettingsManager(const FObjectInitializer& ObjectInitializer = FObjectInitializer::Get());
+	DEVELOPERSETTINGS_API UPlatformSettingsManager(const FObjectInitializer& ObjectInitializer = FObjectInitializer::Get());
 
 	static const UPlatformSettingsManager& Get()
 	{
@@ -50,10 +50,10 @@ public:
 	}
 
 	// Returns the platform settings for the specified class for the current platform
-	UPlatformSettings* GetSettingsForPlatform(TSubclassOf<UPlatformSettings> SettingsClass) const;
+	DEVELOPERSETTINGS_API UPlatformSettings* GetSettingsForPlatform(TSubclassOf<UPlatformSettings> SettingsClass) const;
 
 	// Returns the platform settings for the specified class for the specified platform
-	UPlatformSettings* GetSettingsForPlatform(TSubclassOf<UPlatformSettings> SettingsClass, FName DesiredPlatformIniName) const;
+	DEVELOPERSETTINGS_API UPlatformSettings* GetSettingsForPlatform(TSubclassOf<UPlatformSettings> SettingsClass, FName DesiredPlatformIniName) const;
 
 #if WITH_EDITOR
 	// Gets the simulated platform
@@ -72,7 +72,7 @@ public:
 	}
 
 #if WITH_EDITOR
-	static TArray<FName> GetKnownAndEnablePlatformIniNames();
+	static DEVELOPERSETTINGS_API TArray<FName> GetKnownAndEnablePlatformIniNames();
 
 	template <typename TPlatformSettingsClass>
 	FORCEINLINE TArray<UPlatformSettings*> GetAllPlatformSettings() const
@@ -80,7 +80,7 @@ public:
 		return GetAllPlatformSettings(TPlatformSettingsClass::StaticClass());
 	}
 
-	TArray<UPlatformSettings*> GetAllPlatformSettings(TSubclassOf<UPlatformSettings> SettingsClass) const;
+	DEVELOPERSETTINGS_API TArray<UPlatformSettings*> GetAllPlatformSettings(TSubclassOf<UPlatformSettings> SettingsClass) const;
 
 	template <typename TPlatformSettingsClass>
 	FORCEINLINE TPlatformSettingsClass* GetSettingsForPlatform(FName TargetIniPlatformName) const
@@ -92,7 +92,7 @@ public:
 
 private:
 	// Creates a settings object for the specified class and platform
-	UPlatformSettings* CreateSettingsObjectForPlatform(TSubclassOf<UPlatformSettings> SettingsClass, FName TargetIniPlatformName) const;
+	DEVELOPERSETTINGS_API UPlatformSettings* CreateSettingsObjectForPlatform(TSubclassOf<UPlatformSettings> SettingsClass, FName TargetIniPlatformName) const;
 
 
 private:
@@ -105,6 +105,6 @@ private:
 
 #if WITH_EDITOR
 	// Current simulated platform
-	static FName SimulatedEditorPlatform;
+	static DEVELOPERSETTINGS_API FName SimulatedEditorPlatform;
 #endif
 };
diff --git a/Engine/Source/Runtime/Engine/Classes/Camera/CameraModifier_CameraShake.h b/Engine/Source/Runtime/Engine/Classes/Camera/CameraModifier_CameraShake.h
--- a/Engine/Source/Runtime/Engine/Classes/Camera/CameraModifier_CameraShake.h
+++ b/Engine/Source/Runtime/Engine/Classes/Camera/CameraModifier_CameraShake.h
@@ -225,7 +225,7 @@ protected:
 #endif
 
 private:
-	ENGINE_API void OnPreGarbageCollect();
-	ENGINE_API void RemoveInvalidObjectsFromExpiredPool();
+	void OnPreGarbageCollect();
+	void RemoveInvalidObjectsFromExpiredPool();
 };
 
diff --git a/Engine/Source/Runtime/Engine/Classes/Engine/Level.h b/Engine/Source/Runtime/Engine/Classes/Engine/Level.h
--- a/Engine/Source/Runtime/Engine/Classes/Engine/Level.h
+++ b/Engine/Source/Runtime/Engine/Classes/Engine/Level.h
@@ -717,9 +717,9 @@ public:
 	ENGINE_API static bool GetIsStreamingDisabledFromPackage(FName LevelPackage);
 
 	UE_DEPRECATED(5.3, "GetPartitionedLevelCanBeUsedByLevelInstanceFromAsset is deprecated.")
-	ENGINE_API static bool GetPartitionedLevelCanBeUsedByLevelInstanceFromAsset(const FAssetData& Asset) { return true; }
+	static bool GetPartitionedLevelCanBeUsedByLevelInstanceFromAsset(const FAssetData& Asset) { return true; }
 	UE_DEPRECATED(5.3, "GetPartitionedLevelCanBeUsedByLevelInstanceFromPackage is deprecated.")
-	ENGINE_API static bool GetPartitionedLevelCanBeUsedByLevelInstanceFromPackage(FName LevelPackage) { return true; }
+	static bool GetPartitionedLevelCanBeUsedByLevelInstanceFromPackage(FName LevelPackage) { return true; }
 
 	ENGINE_API static FVector GetLevelInstancePivotOffsetFromAsset(const FAssetData& Asset);
 	ENGINE_API static FVector GetLevelInstancePivotOffsetFromPackage(FName LevelPackage);
diff --git a/Engine/Source/Runtime/Engine/Classes/Sound/QuartzQuantizationUtilities.h b/Engine/Source/Runtime/Engine/Classes/Sound/QuartzQuantizationUtilities.h
--- a/Engine/Source/Runtime/Engine/Classes/Sound/QuartzQuantizationUtilities.h
+++ b/Engine/Source/Runtime/Engine/Classes/Sound/QuartzQuantizationUtilities.h
@@ -536,7 +536,7 @@ namespace Audio
 {
 	// data that is gathered by the AudioThread to get passed from FActiveSound->FMixerSourceVoice
 	// eventually converted to IQuartzQuantizedCommand for the Quantized Command itself
-	struct ENGINE_API FQuartzQuantizedRequestData
+	struct FQuartzQuantizedRequestData
 	{
 		// shared with FQuartzQuantizedCommandInitInfo:
 		FName ClockName;
@@ -550,12 +550,12 @@ namespace Audio
 
 	// data that is passed into IQuartzQuantizedCommand::OnQueued
 	// info that derived classes need can be added here
-	struct ENGINE_API FQuartzQuantizedCommandInitInfo
+	struct FQuartzQuantizedCommandInitInfo
 	{
 		FQuartzQuantizedCommandInitInfo() {}
 
 		// conversion ctor from FQuartzQuantizedRequestData
-		FQuartzQuantizedCommandInitInfo(const FQuartzQuantizedRequestData& RHS, float InSampleRate, int32 InSourceID = INDEX_NONE);
+		ENGINE_API FQuartzQuantizedCommandInitInfo(const FQuartzQuantizedRequestData& RHS, float InSampleRate, int32 InSourceID = INDEX_NONE);
 
 		void SetOwningClockPtr(TSharedPtr<Audio::FQuartzClock> InClockPointer)
 		{
@@ -581,7 +581,7 @@ namespace Audio
 	};
 
 	// base class for quantized commands. Virtual methods called by owning clock.
-	class ENGINE_API IQuartzQuantizedCommand : public FQuartzCrossThreadMessage
+	class IQuartzQuantizedCommand : public FQuartzCrossThreadMessage
 	{
 	public:
 
@@ -592,38 +592,38 @@ namespace Audio
 		virtual ~IQuartzQuantizedCommand() {};
 
 		// allocate a copy of the derived class
-		virtual TSharedPtr<IQuartzQuantizedCommand> GetDeepCopyOfDerivedObject() const;
+		ENGINE_API virtual TSharedPtr<IQuartzQuantizedCommand> GetDeepCopyOfDerivedObject() const;
 
-		void AddSubscriber(FQuartzGameThreadSubscriber InSubscriber);
+		ENGINE_API void AddSubscriber(FQuartzGameThreadSubscriber InSubscriber);
 
 		// Command has reached the AudioRenderThread
-		void OnQueued(const FQuartzQuantizedCommandInitInfo& InCommandInitInfo);
+		ENGINE_API void OnQueued(const FQuartzQuantizedCommandInitInfo& InCommandInitInfo);
 
 		// scheduled (finalize subscriber offsets) - called by FQuartzClock
-		void OnScheduled(const FQuartzClockTickRate& InTickRate);
+		ENGINE_API void OnScheduled(const FQuartzClockTickRate& InTickRate);
 
 		// called during FQuartzClock::Tick() to let us call AboutToStart
 		// at different times for different subscribers
-		void Update(int32 NumFramesUntilDeadline);
+		ENGINE_API void Update(int32 NumFramesUntilDeadline);
 
 		// Perhaps the associated sound failed concurrency and will not be playing
-		void FailedToQueue(FQuartzQuantizedRequestData& InGameThreadData);
+		ENGINE_API void FailedToQueue(FQuartzQuantizedRequestData& InGameThreadData);
 
 		// Called 2x Assumed thread latency before OnFinalCallback()
-		void AboutToStart();
+		ENGINE_API void AboutToStart();
 
 		// Called on the final callback of this event boundary.
 		// InNumFramesLeft is the number of frames into the callback the exact quantized event should take place
-		void OnFinalCallback(int32 InNumFramesLeft);
+		ENGINE_API void OnFinalCallback(int32 InNumFramesLeft);
 
 		// Called if the owning clock gets stopped
-		void OnClockPaused();
+		ENGINE_API void OnClockPaused();
 
 		// Called if the owning clock gets started
-		void OnClockStarted();
+		ENGINE_API void OnClockStarted();
 
 		// Called if the event is cancelled before OnFinalCallback() is called
-		void Cancel();
+		ENGINE_API void Cancel();
 
 
 		//Called if the event type uses an altered amount of frames
@@ -655,21 +655,21 @@ namespace Audio
 
 	// Audio Render Thread Handle to a queued command
 	// Used by AudioMixerSourceVoices to access a pending associated command
-	struct ENGINE_API FQuartzQuantizedCommandHandle
+	struct FQuartzQuantizedCommandHandle
 	{
 		FName OwningClockName;
 		TSharedPtr<IQuartzQuantizedCommand> CommandPtr{ nullptr };
 		FMixerDevice* MixerDevice{ nullptr };
 
 		// Attempts to cancel the command. Returns true if the cancellation was successful.
-		bool Cancel();
+		ENGINE_API bool Cancel();
 
 		// Resets the handle to initial state.
-		void Reset();
+		ENGINE_API void Reset();
 	};
 } // namespace Audio
 
-struct ENGINE_API FAudioComponentCommandInfo
+struct FAudioComponentCommandInfo
 {
 	FAudioComponentCommandInfo() {}
 
diff --git a/Engine/Source/Runtime/Engine/Public/ActorPartition/PartitionActor.h b/Engine/Source/Runtime/Engine/Public/ActorPartition/PartitionActor.h
--- a/Engine/Source/Runtime/Engine/Public/ActorPartition/PartitionActor.h
+++ b/Engine/Source/Runtime/Engine/Public/ActorPartition/PartitionActor.h
@@ -26,7 +26,7 @@ public:
 	ENGINE_API virtual TUniquePtr<class FWorldPartitionActorDesc> CreateClassActorDesc() const override;
 	ENGINE_API virtual uint32 GetDefaultGridSize(UWorld* InWorld) const PURE_VIRTUAL(APartitionActor, return 0;)
 	ENGINE_API virtual bool ShouldIncludeGridSizeInName(UWorld * InWorld, const FActorPartitionIdentifier& InIdentifier) const;
-	ENGINE_API virtual bool ShouldIncludeGridSizeInLabel() const { return false; }
+	virtual bool ShouldIncludeGridSizeInLabel() const { return false; }
 	virtual FGuid GetGridGuid() const { return FGuid(); }
 	ENGINE_API virtual bool IsUserManaged() const override;
 #endif
diff --git a/Engine/Source/Runtime/Engine/Public/HeterogeneousVolumeInterface.h b/Engine/Source/Runtime/Engine/Public/HeterogeneousVolumeInterface.h
--- a/Engine/Source/Runtime/Engine/Public/HeterogeneousVolumeInterface.h
+++ b/Engine/Source/Runtime/Engine/Public/HeterogeneousVolumeInterface.h
@@ -9,7 +9,7 @@
 class IHeterogeneousVolumeInterface
 {
 public:
-	ENGINE_API virtual ~IHeterogeneousVolumeInterface() {}
+	virtual ~IHeterogeneousVolumeInterface() {}
 	virtual const FPrimitiveSceneProxy* GetPrimitiveSceneProxy() const = 0;
 	
 	// Local-Space
@@ -34,14 +34,14 @@ class FPrimitiveSceneProxy;
 class FHeterogeneousVolumeData : public IHeterogeneousVolumeInterface, FOneFrameResource
 {
 public:
-	ENGINE_API explicit FHeterogeneousVolumeData(const FPrimitiveSceneProxy* SceneProxy)
+	explicit FHeterogeneousVolumeData(const FPrimitiveSceneProxy* SceneProxy)
 		: PrimitiveSceneProxy(SceneProxy)
 		, InstanceToLocal(FMatrix::Identity)
 		, VoxelResolution(FIntVector::ZeroValue)
 		, MinimumVoxelSize(0.1)
 		, LightingDownsampleFactor(1.0)
 	{}
-	ENGINE_API virtual ~FHeterogeneousVolumeData() {}
+	virtual ~FHeterogeneousVolumeData() {}
 
 	virtual const FPrimitiveSceneProxy* GetPrimitiveSceneProxy() const { return PrimitiveSceneProxy; }
 
diff --git a/Engine/Source/Runtime/Engine/Public/LevelInstance/LevelInstanceSubsystem.h b/Engine/Source/Runtime/Engine/Public/LevelInstance/LevelInstanceSubsystem.h
--- a/Engine/Source/Runtime/Engine/Public/LevelInstance/LevelInstanceSubsystem.h
+++ b/Engine/Source/Runtime/Engine/Public/LevelInstance/LevelInstanceSubsystem.h
@@ -69,9 +69,9 @@ public:
 	ENGINE_API bool OnExitEditorModeInternal(bool bForceExit);
 
 	UE_DEPRECATED(5.3, "Use FPackedLevelActorUtils::PackAllLoadedActors")
-	ENGINE_API void PackAllLoadedActors() {}
+	void PackAllLoadedActors() {}
 	UE_DEPRECATED(5.3, "Use FPackedLevelActorUtils::CanPackAllLoadedActors")
-	ENGINE_API bool CanPackAllLoadedActors() const { return false;}
+	bool CanPackAllLoadedActors() const { return false;}
 
 	ENGINE_API ILevelInstanceInterface* GetEditingLevelInstance() const;
 	ENGINE_API bool CanEditLevelInstance(const ILevelInstanceInterface* LevelInstance, FText* OutReason = nullptr) const;
@@ -124,7 +124,7 @@ public:
 	static ENGINE_API bool CanUseWorldAsset(const ILevelInstanceInterface* LevelInstance, TSoftObjectPtr<UWorld> WorldAsset, FString* OutReason);
 	
 	UE_DEPRECATED(5.3, "CanUsePackage is deprecated.")
-	static ENGINE_API bool CanUsePackage(FName InPackageName) { return true;  }
+	static bool CanUsePackage(FName InPackageName) { return true;  }
 
 	/** Editor-only event triggered when level instance is committed with changes */
 	DECLARE_EVENT_OneParam(ULevelInstanceSubsystem, FLevelInstanceChanged, FName);
@@ -239,4 +239,4 @@ private:
 
 	FWorldPartitionReference CurrentEditLevelInstanceActor;
 #endif
-};
\ No newline at end of file
+};
diff --git a/Engine/Source/Runtime/Engine/Public/PackedLevelActor/PackedLevelActor.h b/Engine/Source/Runtime/Engine/Public/PackedLevelActor/PackedLevelActor.h
--- a/Engine/Source/Runtime/Engine/Public/PackedLevelActor/PackedLevelActor.h
+++ b/Engine/Source/Runtime/Engine/Public/PackedLevelActor/PackedLevelActor.h
@@ -44,11 +44,11 @@ public:
 	ENGINE_API virtual TUniquePtr<class FWorldPartitionActorDesc> CreateClassActorDesc() const override;
 	
 	UE_DEPRECATED(5.3, "Use FPackedLevelActorUtils::CreateOrUpdateBlueprint instead")
-	static ENGINE_API bool CreateOrUpdateBlueprint(ALevelInstance* InLevelInstance, TSoftObjectPtr<UBlueprint> InBlueprintAsset, bool bCheckoutAndSave = true, bool bPromptForSave = true) { return false;}
+	static bool CreateOrUpdateBlueprint(ALevelInstance* InLevelInstance, TSoftObjectPtr<UBlueprint> InBlueprintAsset, bool bCheckoutAndSave = true, bool bPromptForSave = true) { return false;}
 	UE_DEPRECATED(5.3, "Use FPackedLevelActorUtils::CreateOrUpdateBlueprint instead")
-	static ENGINE_API bool CreateOrUpdateBlueprint(TSoftObjectPtr<UWorld> InWorldAsset, TSoftObjectPtr<UBlueprint> InBlueprintAsset, bool bCheckoutAndSave = true, bool bPromptForSave = true) { return false;}
+	static bool CreateOrUpdateBlueprint(TSoftObjectPtr<UWorld> InWorldAsset, TSoftObjectPtr<UBlueprint> InBlueprintAsset, bool bCheckoutAndSave = true, bool bPromptForSave = true) { return false;}
 	UE_DEPRECATED(5.3, "Use FPackedLevelActorUtils::UpdateBlueprint instead")
-	static ENGINE_API void UpdateBlueprint(UBlueprint* InBlueprint, bool bCheckoutAndSave = true) {}
+	static void UpdateBlueprint(UBlueprint* InBlueprint, bool bCheckoutAndSave = true) {}
 
 	static ENGINE_API FName GetPackedComponentTag();
 
diff --git a/Engine/Source/Runtime/Engine/Public/Scalability.h b/Engine/Source/Runtime/Engine/Public/Scalability.h
--- a/Engine/Source/Runtime/Engine/Public/Scalability.h
+++ b/Engine/Source/Runtime/Engine/Public/Scalability.h
@@ -155,7 +155,7 @@ namespace Scalability
 	ENGINE_API extern FOnScalabilitySettingsChanged OnScalabilitySettingsChanged;
 
 	/** Structure holding the details of a preset. */
-	struct ENGINE_API FResolutionPreset
+	struct FResolutionPreset
 	{
 		FString Name;
 		int32 Id = 0;
diff --git a/Engine/Source/Runtime/Engine/Public/WorldPartition/WorldPartition.h b/Engine/Source/Runtime/Engine/Public/WorldPartition/WorldPartition.h
--- a/Engine/Source/Runtime/Engine/Public/WorldPartition/WorldPartition.h
+++ b/Engine/Source/Runtime/Engine/Public/WorldPartition/WorldPartition.h
@@ -125,9 +125,9 @@ public:
 	ENGINE_API void SetEnableStreaming(bool bInEnableStreaming);
 
 	UE_DEPRECATED(5.3, "CanBeUsedByLevelInstance is deprecated.")
-	ENGINE_API bool CanBeUsedByLevelInstance() const { return true; }
+	bool CanBeUsedByLevelInstance() const { return true; }
 	UE_DEPRECATED(5.3, "SetCanBeUsedByLevelInstance is deprecated.")
-	ENGINE_API void SetCanBeUsedByLevelInstance(bool bInCanBeUsedByLevelInstance) {}
+	void SetCanBeUsedByLevelInstance(bool bInCanBeUsedByLevelInstance) {}
 
 	ENGINE_API void OnEnableStreamingChanged();
 
@@ -536,4 +536,4 @@ private:
 
 	friend class AWorldPartitionReplay;
 	friend class UWorldPartitionSubsystem;
-};
\ No newline at end of file
+};
diff --git a/Engine/Source/Runtime/EngineSettings/Classes/ConsoleSettings.h b/Engine/Source/Runtime/EngineSettings/Classes/ConsoleSettings.h
--- a/Engine/Source/Runtime/EngineSettings/Classes/ConsoleSettings.h
+++ b/Engine/Source/Runtime/EngineSettings/Classes/ConsoleSettings.h
@@ -60,8 +60,8 @@ struct FAutoCompleteCommand
 /**
  * Implements the settings for the UConsole class.
  */
-UCLASS(config=Input, defaultconfig)
-class ENGINESETTINGS_API UConsoleSettings
+UCLASS(config=Input, defaultconfig, MinimalAPI)
+class UConsoleSettings
 	: public UObject
 {
 	GENERATED_UCLASS_BODY()
diff --git a/Engine/Source/Runtime/EngineSettings/Classes/GameMapsSettings.h b/Engine/Source/Runtime/EngineSettings/Classes/GameMapsSettings.h
--- a/Engine/Source/Runtime/EngineSettings/Classes/GameMapsSettings.h
+++ b/Engine/Source/Runtime/EngineSettings/Classes/GameMapsSettings.h
@@ -91,8 +91,8 @@ struct FTemplateMapInfoOverride
 };
 
 
-UCLASS(config=Engine, defaultconfig)
-class ENGINESETTINGS_API UGameMapsSettings
+UCLASS(config=Engine, defaultconfig, MinimalAPI)
+class UGameMapsSettings
 	: public UObject
 {
 	GENERATED_UCLASS_BODY()
@@ -103,7 +103,7 @@ class ENGINESETTINGS_API UGameMapsSettings
 	 *
 	 * @return the default map specified in the settings
 	 */
-	static FString GetGameDefaultMap( );
+	static ENGINESETTINGS_API FString GetGameDefaultMap( );
 
 	/**
 	 * Get the global default game type specified in the configuration
@@ -111,38 +111,38 @@ class ENGINESETTINGS_API UGameMapsSettings
 	 * 
 	 * @return the proper global default game type
 	 */
-	static FString GetGlobalDefaultGameMode( );
+	static ENGINESETTINGS_API FString GetGlobalDefaultGameMode( );
 
 	/**
 	 * Searches the GameModeClassAliases list for a named game mode, if not found will return passed in string
 	 * 
 	 * @return the proper game type class path to load
 	 */
-	static FString GetGameModeForName( const FString& GameModeName );
+	static ENGINESETTINGS_API FString GetGameModeForName( const FString& GameModeName );
 
 	/**
 	 * Searches the GameModeMapPrefixes list for a named game mode, if not found will return passed in string
 	 * 
 	 * @return the proper game type class path to load, or empty if not found
 	 */
-	static FString GetGameModeForMapName( const FString& MapName );
+	static ENGINESETTINGS_API FString GetGameModeForMapName( const FString& MapName );
 
 	/**
 	 * Set the default map to use (see GameDefaultMap below)
 	 *
 	 * @param NewMap name of valid map to use
 	 */
-	static void SetGameDefaultMap( const FString& NewMap );
+	static ENGINESETTINGS_API void SetGameDefaultMap( const FString& NewMap );
 
 	/**
 	 * Set the default game type (see GlobalDefaultGameMode below)
 	 *
 	 * @param NewGameMode name of valid map to use
 	 */
-	static void SetGlobalDefaultGameMode( const FString& NewGameMode );
+	static ENGINESETTINGS_API void SetGlobalDefaultGameMode( const FString& NewGameMode );
 
-	virtual void PostInitProperties() override;
-	virtual void PostReloadConfig( class FProperty* PropertyThatWasLoaded ) override;
+	ENGINESETTINGS_API virtual void PostInitProperties() override;
+	ENGINESETTINGS_API virtual void PostReloadConfig( class FProperty* PropertyThatWasLoaded ) override;
 
 public:
 
@@ -224,7 +224,7 @@ public:
 
 	/** Returns the game local maps settings */
 	UFUNCTION(BlueprintPure, Category = Settings, meta=(DisplayName="Get Game Maps and Modes Settings"))
-	static UGameMapsSettings* GetGameMapsSettings();
+	static ENGINESETTINGS_API UGameMapsSettings* GetGameMapsSettings();
 
 	/**
 	 * Modify "Skip Assigning Gamepad to Player 1" GameMapsSettings option
@@ -232,9 +232,9 @@ public:
 	 * @note This value is saved to local config when changed.
 	*/
 	UFUNCTION(BlueprintCallable, Category = "Utilities")
-	void SetSkipAssigningGamepadToPlayer1(bool bSkipFirstPlayer = true);
+	ENGINESETTINGS_API void SetSkipAssigningGamepadToPlayer1(bool bSkipFirstPlayer = true);
 
 	UFUNCTION(BlueprintPure, Category = "Utilities")
-	bool GetSkipAssigningGamepadToPlayer1() const;
+	ENGINESETTINGS_API bool GetSkipAssigningGamepadToPlayer1() const;
 
 };
diff --git a/Engine/Source/Runtime/EngineSettings/Classes/GameNetworkManagerSettings.h b/Engine/Source/Runtime/EngineSettings/Classes/GameNetworkManagerSettings.h
--- a/Engine/Source/Runtime/EngineSettings/Classes/GameNetworkManagerSettings.h
+++ b/Engine/Source/Runtime/EngineSettings/Classes/GameNetworkManagerSettings.h
@@ -11,8 +11,8 @@
 /**
  * Holds the settings for the AGameNetworkManager class.
  */
-UCLASS(config=Game)
-class ENGINESETTINGS_API UGameNetworkManagerSettings
+UCLASS(config=Game, MinimalAPI)
+class UGameNetworkManagerSettings
 	: public UObject
 {
 	GENERATED_UCLASS_BODY()
diff --git a/Engine/Source/Runtime/EngineSettings/Classes/GameSessionSettings.h b/Engine/Source/Runtime/EngineSettings/Classes/GameSessionSettings.h
--- a/Engine/Source/Runtime/EngineSettings/Classes/GameSessionSettings.h
+++ b/Engine/Source/Runtime/EngineSettings/Classes/GameSessionSettings.h
@@ -8,8 +8,8 @@
 #include "GameSessionSettings.generated.h"
 
 
-UCLASS(config=Game)
-class ENGINESETTINGS_API UGameSessionSettings
+UCLASS(config=Game, MinimalAPI)
+class UGameSessionSettings
 	: public UObject
 {
 	GENERATED_UCLASS_BODY()
diff --git a/Engine/Source/Runtime/EngineSettings/Classes/GeneralEngineSettings.h b/Engine/Source/Runtime/EngineSettings/Classes/GeneralEngineSettings.h
--- a/Engine/Source/Runtime/EngineSettings/Classes/GeneralEngineSettings.h
+++ b/Engine/Source/Runtime/EngineSettings/Classes/GeneralEngineSettings.h
@@ -8,8 +8,8 @@
 #include "GeneralEngineSettings.generated.h"
 
 
-UCLASS(config=Engine)
-class ENGINESETTINGS_API UGeneralEngineSettings
+UCLASS(config=Engine, MinimalAPI)
+class UGeneralEngineSettings
 	: public UObject
 {
 	GENERATED_UCLASS_BODY()
diff --git a/Engine/Source/Runtime/EngineSettings/Classes/GeneralProjectSettings.h b/Engine/Source/Runtime/EngineSettings/Classes/GeneralProjectSettings.h
--- a/Engine/Source/Runtime/EngineSettings/Classes/GeneralProjectSettings.h
+++ b/Engine/Source/Runtime/EngineSettings/Classes/GeneralProjectSettings.h
@@ -9,8 +9,8 @@
 #include "GeneralProjectSettings.generated.h"
 
 
-UCLASS(config=Game, defaultconfig)
-class ENGINESETTINGS_API UGeneralProjectSettings
+UCLASS(config=Game, defaultconfig, MinimalAPI)
+class UGeneralProjectSettings
 	: public UObject
 {
 	GENERATED_UCLASS_BODY()
diff --git a/Engine/Source/Runtime/EngineSettings/Classes/HudSettings.h b/Engine/Source/Runtime/EngineSettings/Classes/HudSettings.h
--- a/Engine/Source/Runtime/EngineSettings/Classes/HudSettings.h
+++ b/Engine/Source/Runtime/EngineSettings/Classes/HudSettings.h
@@ -8,8 +8,8 @@
 #include "HudSettings.generated.h"
 
 
-UCLASS(config=Game)
-class ENGINESETTINGS_API UHudSettings
+UCLASS(config=Game, MinimalAPI)
+class UHudSettings
 	: public UObject
 {
 	GENERATED_UCLASS_BODY()
