Cloth - Fixed a crash in the skeletal mesh cloth painting tool caused by invalid bone mapping indices after importing an existing cloth onto a different skeleton.

diff --git a/Engine/Source/Runtime/ClothingSystemRuntimeCommon/Private/ClothingAsset.cpp b/Engine/Source/Runtime/ClothingSystemRuntimeCommon/Private/ClothingAsset.cpp
--- a/Engine/Source/Runtime/ClothingSystemRuntimeCommon/Private/ClothingAsset.cpp
+++ b/Engine/Source/Runtime/ClothingSystemRuntimeCommon/Private/ClothingAsset.cpp
@@ -1013,9 +1013,8 @@ void UClothingAssetCommon::RefreshBoneMapping(USkeletalMesh* InSkelMesh)
 	// Repopulate the used indices.
 	for(int32 BoneNameIndex = 0; BoneNameIndex < UsedBoneNames.Num(); ++BoneNameIndex)
 	{
-		UsedBoneIndices[BoneNameIndex] = 
-			InSkelMesh->GetRefSkeleton().FindBoneIndex(
-				UsedBoneNames[BoneNameIndex]);
+		const int32 BoneIndex = InSkelMesh->GetRefSkeleton().FindBoneIndex(UsedBoneNames[BoneNameIndex]);
+		UsedBoneIndices[BoneNameIndex] = (BoneIndex != INDEX_NONE) ? BoneIndex : 0;  // Remap to the root bone if the two skeletons differ
 	}
 }
 
