Chaos Cloth - Fixed the shared config outer ownership when adding cloth assets from another skeletal mesh.

diff --git a/Engine/Source/Runtime/ClothingSystemRuntimeCommon/Private/ClothingAsset.cpp b/Engine/Source/Runtime/ClothingSystemRuntimeCommon/Private/ClothingAsset.cpp
--- a/Engine/Source/Runtime/ClothingSystemRuntimeCommon/Private/ClothingAsset.cpp
+++ b/Engine/Source/Runtime/ClothingSystemRuntimeCommon/Private/ClothingAsset.cpp
@@ -1484,7 +1484,7 @@ void UClothingAssetCommon::PropagateSharedConfigs(bool bMigrateSharedConfigToCon
 		}
 
 		// Migrate the common shared configs' deprecated parameters to all per cloth configs, and fix the shared config ownership
-		for (const TPair<FName, TObjectPtr<UClothConfigBase>>& ClothSharedConfigItem : ClothConfigs)
+		for (TPair<FName, TObjectPtr<UClothConfigBase>>& ClothSharedConfigItem : ClothConfigs)
 		{
 			if (UClothSharedConfigCommon* const ClothSharedConfig = Cast<UClothSharedConfigCommon>(ClothSharedConfigItem.Value))
 			{
@@ -1503,11 +1503,17 @@ void UClothingAssetCommon::PropagateSharedConfigs(bool bMigrateSharedConfigToCon
 						}
 					}
 				}
-				// Fix the shared config outer if it is still a common asset (the config must belong to the skeletal mesh, as it is shared between assets)
-				if (ClothSharedConfig->GetOuter() != SkeletalMesh)
+				// Fix the shared config outer if it is still a UClothingAssetCommon (the config must belong to the skeletal mesh, as it is shared between assets)
+				if (Cast<UClothingAssetCommon>(ClothSharedConfig->GetOuter()))
 				{
 					ClothSharedConfig->Rename(nullptr, SkeletalMesh, REN_DontCreateRedirectors | REN_ForceNoResetLoaders | REN_NonTransactional);
 				}
+
+				// Fix the shared config ownership, asset might have been copied and the shared config could still point to a different skeletal mesh
+				if (ClothSharedConfig->GetOuter() != SkeletalMesh)
+				{
+					ClothSharedConfigItem.Value = DuplicateObject<UClothSharedConfigCommon>(ClothSharedConfig, SkeletalMesh, ClothSharedConfig->GetFName());
+				}
 			}
 		}
 	}
