Cloth - Fixed an edge case when adding clothing data from a different skeletal mesh asset crashes, due to recalculating the reference bone index in legacy assets without having the correct bone mapping.

#rnx

diff --git a/Engine/Source/Runtime/ClothingSystemRuntimeCommon/Private/ClothingAsset.cpp b/Engine/Source/Runtime/ClothingSystemRuntimeCommon/Private/ClothingAsset.cpp
--- a/Engine/Source/Runtime/ClothingSystemRuntimeCommon/Private/ClothingAsset.cpp
+++ b/Engine/Source/Runtime/ClothingSystemRuntimeCommon/Private/ClothingAsset.cpp
@@ -1326,6 +1326,8 @@ void UClothingAssetCommon::PostLoad()
 		ClothingCachedDataFlags |= EClothingCachedDataFlagsCommon::Tethers;
 
 		// ReferenceBoneIndex is only required when rebinding the cloth.
+		USkeletalMesh* const OwnerMesh = CastChecked<USkeletalMesh>(GetOuter());
+		RefreshBoneMapping(OwnerMesh);
 		CalculateReferenceBoneIndex();
 	}
 
