Cloth: Add new SelfCollisionSpheres to legacy skeletal mesh cloth assets.

I have deprecated the existing array SelfCollisionIndices (they're unused now that apex cloth is no longer supported) on the ClothPhysicalMeshData for a new Set SelfCollisionVertexSet, but it's built in the same way.

diff --git a/Engine/Plugins/Experimental/ChaosCloth/Source/Chaos/Private/ChaosCloth/ChaosClothingSimulationConfig.cpp b/Engine/Plugins/Experimental/ChaosCloth/Source/Chaos/Private/ChaosCloth/ChaosClothingSimulationConfig.cpp
--- a/Engine/Plugins/Experimental/ChaosCloth/Source/Chaos/Private/ChaosCloth/ChaosClothingSimulationConfig.cpp
+++ b/Engine/Plugins/Experimental/ChaosCloth/Source/Chaos/Private/ChaosCloth/ChaosClothingSimulationConfig.cpp
@@ -194,6 +194,12 @@ namespace Chaos
 			Property->AddValue(TEXT("UseSelfIntersections"), ClothConfig->bUseSelfIntersections, NonAnimatablePropertyFlags);
 			Property->AddValue(TEXT("SelfCollisionFriction"), ClothConfig->SelfCollisionFriction, NonAnimatablePropertyFlags);
 
+			if (ClothConfig->bUseSelfCollisionSpheres)
+			{
+				Property->AddValue("SelfCollisionSphereRadius", ClothConfig->SelfCollisionSphereRadius, NonAnimatablePropertyFlags);
+				Property->AddValue("SelfCollisionSphereStiffness", ClothConfig->SelfCollisionSphereStiffness, AnimatablePropertyFlags);
+			}
+
 			// Max distance
 			{
 				const int32 MaxDistanceIndex = Property->AddProperty(TEXT("MaxDistance"), AnimatablePropertyFlags);
diff --git a/Engine/Plugins/Experimental/ChaosCloth/Source/Chaos/Private/ChaosCloth/ChaosClothingSimulationSkeletalMesh.cpp b/Engine/Plugins/Experimental/ChaosCloth/Source/Chaos/Private/ChaosCloth/ChaosClothingSimulationSkeletalMesh.cpp
--- a/Engine/Plugins/Experimental/ChaosCloth/Source/Chaos/Private/ChaosCloth/ChaosClothingSimulationSkeletalMesh.cpp
+++ b/Engine/Plugins/Experimental/ChaosCloth/Source/Chaos/Private/ChaosCloth/ChaosClothingSimulationSkeletalMesh.cpp
@@ -161,8 +161,16 @@ PRAGMA_ENABLE_DEPRECATION_WARNINGS
 
 	TMap<FString, const TSet<int32>*> FClothingSimulationSkeletalMesh::GetVertexSets(int32 LODIndex) const
 	{
-		// Not supported
-		return TMap<FString, const TSet<int32>*>();
+		TMap<FString, const TSet<int32>*> VertexSets;
+		if (IsValidLODIndex(LODIndex))
+		{
+			const FClothLODDataCommon& ClothLODData = Asset->LodData[LODIndex];
+			const FClothPhysicalMeshData& ClothPhysicalMeshData = ClothLODData.PhysicalMeshData;
+			static const FString SelfCollisionSetName("_SelfCollisionSpheres");
+			VertexSets.Add(SelfCollisionSetName, &ClothPhysicalMeshData.SelfCollisionVertexSet);
+		}
+
+		return VertexSets;
 	}
 
 	TArray<TConstArrayView<TTuple<int32, int32, float>>> FClothingSimulationSkeletalMesh::GetTethers(int32 LODIndex, bool bUseGeodesicTethers) const
diff --git a/Engine/Plugins/Experimental/ChaosCloth/Source/Chaos/Public/ChaosCloth/ChaosClothConfig.h b/Engine/Plugins/Experimental/ChaosCloth/Source/Chaos/Public/ChaosCloth/ChaosClothConfig.h
--- a/Engine/Plugins/Experimental/ChaosCloth/Source/Chaos/Public/ChaosCloth/ChaosClothConfig.h
+++ b/Engine/Plugins/Experimental/ChaosCloth/Source/Chaos/Public/ChaosCloth/ChaosClothConfig.h
@@ -70,12 +70,17 @@ public:
 	/** PostLoad override used to deal with updates/changes in properties. */
 	CHAOSCLOTH_API virtual void PostLoad() override;
 
+	virtual bool NeedsSelfCollisionData() const override { return bUseSelfCollisionSpheres && SelfCollisionSphereRadius * SelfCollisionSphereRadiusCullMultiplier > 0.f; }
+
 	/** Return wherether to pre-compute Inverse Masses. */
 	virtual bool NeedsInverseMasses() const override { return false; }  // TODO: Chaos Cloth uses the mass mode enum, and this will require a little refactor to work
 
 	/** Return wherether to pre-compute the Long Range Attachment tethers. */
 	virtual bool NeedsTethers() const override { return TetherStiffness.Low > 0.f || TetherStiffness.High > 0.f; }
 
+	/** Return the self collision radius to precomute self collision data. */
+	virtual float GetSelfCollisionRadius() const override { return SelfCollisionSphereRadius * SelfCollisionSphereRadiusCullMultiplier; }
+
 	/** Return whether tethers need to be calculated using geodesic distances instead of eclidean. */
 	virtual bool TethersUseGeodesicDistance() const override { return bUseGeodesicDistance; }
 
@@ -235,6 +240,25 @@ public:
 	UPROPERTY(EditAnywhere, Category = "Collision Properties", meta = (EditCondition = "bUseSelfCollisions"))
 	bool bUseSelfIntersections = false;
 
+	/** Enable sphere-based self collision repulsion forces. */
+	UPROPERTY(EditAnywhere, Category = "Collision Properties",meta = (EditCondition = "!bUseSelfCollisions"))
+	bool bUseSelfCollisionSpheres = false;
+
+	/** The radius of the spheres used in self collision centered at each vertex. */
+	UPROPERTY(EditAnywhere, Category = "Collision Properties", meta = (EditCondition = "!bUseSelfCollisions && bUseSelfCollisionSpheres", UIMin = "0", UIMax = "100", ClampMin = "0", ClampMax = "1000"))
+	float SelfCollisionSphereRadius = 0.5f;
+
+	/** The stiffness of the springs used to control self collision. */
+	UPROPERTY(EditAnywhere, Category = "Collision Properties", meta = (EditCondition = "!bUseSelfCollisions && bUseSelfCollisionSpheres", UIMin = "0", UIMax = "100", ClampMin = "0", ClampMax = "1000"))
+	float SelfCollisionSphereStiffness = 1.f;
+
+	/** 
+	 * Multiplier for culling the self collision spheres. Spheres are seeded on every vertex,
+	 * and culled based on SelfCollisionSphereRadius * SelfCollisionSphereRadiusCullMultiplier.
+	 */
+	UPROPERTY(EditAnywhere, Category = "Collision Properties", meta = (EditCondition = "!bUseSelfCollisions && bUseSelfCollisionSpheres", ClampMin = "1"))
+	float SelfCollisionSphereRadiusCullMultiplier = 1.f;
+
 	/**
 	 * This parameter is automatically set by the migration code. It can be overridden here to use the old way of authoring the backstop distances.
 	 * The legacy backstop requires the sphere radius to be included within the painted distance mask, making it difficult to author correctly. In this case the backstop distance is the distance from the animated mesh to the center of the corresponding backstop collision sphere.
diff --git a/Engine/Plugins/Experimental/Mutable/Source/CustomizableObject/Private/MuCO/CustomizableObjectInstance.cpp b/Engine/Plugins/Experimental/Mutable/Source/CustomizableObject/Private/MuCO/CustomizableObjectInstance.cpp
--- a/Engine/Plugins/Experimental/Mutable/Source/CustomizableObject/Private/MuCO/CustomizableObjectInstance.cpp
+++ b/Engine/Plugins/Experimental/Mutable/Source/CustomizableObject/Private/MuCO/CustomizableObjectInstance.cpp
@@ -4329,38 +4329,29 @@ void UCustomizableInstancePrivateData::BuildOrCopyClothingData(const TSharedPtr<
 			NewPhysicalMeshIndices.SetNum(SrcPhysicalMeshIndices.Num());
 			NewPhysicalMeshIndices.SetNum(TrimAndRemapTriangles(NewPhysicalMeshIndices, SrcPhysicalMeshIndices), false);
 		
-			const auto TrimAndRemapIndices = [&IndexMap](TArray<uint32>& Dst, const TArray<uint32>& Src) -> int32
+			const auto TrimAndRemapVertexSet = [&IndexMap](TSet<int32>& Dst, const TSet<int32>& Src)
 			{	
-				const int32 SrcNumElems = Src.Num();
 				if (!IndexMap.Num())
 				{
-					//for (int32 Idx = 0; Idx < SrcNumElems; ++Idx)
-					//{
-					//	Dst[Idx] = Src[Idx];
-					//}
-
-					FMemory::Memcpy( Dst.GetData(), Src.GetData(), SrcNumElems*sizeof(uint32) );
-					return SrcNumElems;
+					Dst = Src;
+					return;
 				}
 
-				int32 DstNumElems = 0;
-				for (int32 Idx = 0; Idx < SrcNumElems; ++Idx)
+				Dst.Reserve(Src.Num());
+				for(const int32 SrcIdx : Src)
 				{
-					const int32 MappedIdx = IndexMap[Src[Idx]];
+					const int32 MappedIdx = IndexMap[SrcIdx];
 
 					if (MappedIdx >= 0)
 					{
-						Dst[DstNumElems++] = MappedIdx;
+						Dst.Add(MappedIdx);
 					}
 				}
-
-				return DstNumElems;
 			};
 
-			const TArray<uint32>& SrcSelfCollisionIndices = SrcLodData.PhysicalMeshData.SelfCollisionIndices;
-			TArray<uint32>& NewSelfCollisionIndices = NewLodData.PhysicalMeshData.SelfCollisionIndices;
-			NewSelfCollisionIndices.SetNum(SrcSelfCollisionIndices.Num());
-			NewSelfCollisionIndices.SetNum(TrimAndRemapIndices(NewSelfCollisionIndices, SrcSelfCollisionIndices), false);
+			const TSet<int32>& SrcSelfCollisionVertexSet = SrcLodData.PhysicalMeshData.SelfCollisionVertexSet;
+			TSet<int32>& NewSelfCollisionVertexSet = NewLodData.PhysicalMeshData.SelfCollisionVertexSet;
+			TrimAndRemapVertexSet(NewSelfCollisionVertexSet, SrcSelfCollisionVertexSet);;
 						
 			{
 				MUTABLE_CPUPROFILER_SCOPE(BuildClothTetherData)
diff --git a/Engine/Source/Runtime/ClothingSystemRuntimeCommon/Private/ClothPhysicalMeshData.cpp b/Engine/Source/Runtime/ClothingSystemRuntimeCommon/Private/ClothPhysicalMeshData.cpp
--- a/Engine/Source/Runtime/ClothingSystemRuntimeCommon/Private/ClothPhysicalMeshData.cpp
+++ b/Engine/Source/Runtime/ClothingSystemRuntimeCommon/Private/ClothPhysicalMeshData.cpp
@@ -29,7 +29,7 @@ void FClothPhysicalMeshData::MigrateFrom(FClothPhysicalMeshData& ClothPhysicalMe
 		BoneData = MoveTemp(ClothPhysicalMeshData.BoneData);
 		NumFixedVerts = ClothPhysicalMeshData.NumFixedVerts;
 		MaxBoneWeights = ClothPhysicalMeshData.MaxBoneWeights;
-		SelfCollisionIndices = MoveTemp(ClothPhysicalMeshData.SelfCollisionIndices);
+		SelfCollisionVertexSet = MoveTemp(ClothPhysicalMeshData.SelfCollisionVertexSet);
 	}
 }
 
@@ -45,7 +45,11 @@ void FClothPhysicalMeshData::MigrateFrom(UClothPhysicalMeshDataBase_Legacy* Clot
 	BoneData = MoveTemp(ClothPhysicalMeshDataBase->BoneData);
 	NumFixedVerts = ClothPhysicalMeshDataBase->NumFixedVerts;
 	MaxBoneWeights = ClothPhysicalMeshDataBase->MaxBoneWeights;
-	SelfCollisionIndices = MoveTemp(ClothPhysicalMeshDataBase->SelfCollisionIndices);
+	SelfCollisionVertexSet.Empty(ClothPhysicalMeshDataBase->SelfCollisionIndices.Num());
+	for (uint32 Index : ClothPhysicalMeshDataBase->SelfCollisionIndices)
+	{
+		SelfCollisionVertexSet.Add((int32)Index);
+	}
 
 	const TArray<uint32> FloatArrayIds = ClothPhysicalMeshDataBase->GetFloatArrayIds();
 	for (uint32 FloatArrayId : FloatArrayIds)
@@ -69,6 +73,8 @@ void FClothPhysicalMeshData::Reset(const int32 InNumVerts, const int32 InNumIndi
 	BoneData.AddDefaulted(InNumVerts);
 	Indices.Init(0, InNumIndices);
 
+	SelfCollisionVertexSet.Reset();
+
 	NumFixedVerts = 0;
 	MaxBoneWeights = 0;
 
@@ -103,65 +109,56 @@ void FClothPhysicalMeshData::BuildSelfCollisionData(const TMap<FName, TObjectPtr
 
 void FClothPhysicalMeshData::BuildSelfCollisionData(float SelfCollisionRadius)
 {
-	const float SelfCollisionRadiusSq = SelfCollisionRadius * SelfCollisionRadius;
+	const float SelfCollisionDiamSq = 4.f *SelfCollisionRadius * SelfCollisionRadius;
 
 	// Start with the full set
 	const int32 NumVerts = Vertices.Num();
-	SelfCollisionIndices.Reset();
+	TArray<bool> VertexIsValid;
+	VertexIsValid.Init(true, NumVerts);
+
 	const FPointWeightMap& MaxDistances = GetWeightMap(EWeightMapTargetCommon::MaxDistance);
 	for (int32 Index = 0; Index < NumVerts; ++Index)
 	{
-		if (!MaxDistances.IsBelowThreshold(Index))
+		if (MaxDistances.IsBelowThreshold(Index))
 		{
-			SelfCollisionIndices.Add(Index);
+			VertexIsValid[Index] = false;
 		}
 	}
 
 	// Now start aggressively culling verts that are near others that we have accepted
-	for (int32 Vert0Itr = 0; Vert0Itr < SelfCollisionIndices.Num(); ++Vert0Itr)
+	SelfCollisionVertexSet.Reset();
+	for (int32 V0Index = 0; V0Index < NumVerts; ++V0Index)
 	{
-		const uint32 V0Index = SelfCollisionIndices[Vert0Itr];
-		if (V0Index == INDEX_NONE)
+		if (!VertexIsValid[V0Index])
 		{
-			// We'll remove these indices later.  Just skip it for now.
 			continue;
 		}
+		SelfCollisionVertexSet.Add(V0Index);
 
 		const FVector& V0Pos = (FVector)Vertices[V0Index];
 
 		// Start one after our current V0, we've done the other checks
-		for (int32 Vert1Itr = Vert0Itr + 1; Vert1Itr < SelfCollisionIndices.Num(); ++Vert1Itr)
+		for (int32 V1Index = V0Index + 1; V1Index < NumVerts; ++V1Index)
 		{
-			const uint32 V1Index = SelfCollisionIndices[Vert1Itr];
-			if (V1Index == INDEX_NONE)
+			if (!VertexIsValid[V1Index])
 			{
-				// We'll remove these indices later.  Just skip it for now.
 				continue;
 			}
 
 			const FVector& V1Pos = (FVector)Vertices[V1Index];
 			const float V0ToV1DistSq = (V1Pos - V0Pos).SizeSquared();
-			if (V0ToV1DistSq < SelfCollisionRadiusSq)
+			if (V0ToV1DistSq < SelfCollisionDiamSq)
 			{
 				// Points are in contact in the rest state.  Remove it.
 				//
 				// It's worth noting that this biases towards removing indices 
 				// of later in the list, and keeping ones earlier.  That's not
 				// a great criteria for choosing which one is more important.
-				SelfCollisionIndices[Vert1Itr] = INDEX_NONE;
+				VertexIsValid[V1Index] = false;
 				continue;
 			}
 		}
 	}
-
-	// Cull flagged indices.
-	for (int32 It = SelfCollisionIndices.Num(); It--;)
-	{
-		if (SelfCollisionIndices[It] == INDEX_NONE)
-		{
-			SelfCollisionIndices.RemoveAt(It);
-		}
-	}
 }
 
 void FClothPhysicalMeshData::CalculateInverseMasses()
diff --git a/Engine/Source/Runtime/ClothingSystemRuntimeCommon/Private/ClothingAsset.cpp b/Engine/Source/Runtime/ClothingSystemRuntimeCommon/Private/ClothingAsset.cpp
--- a/Engine/Source/Runtime/ClothingSystemRuntimeCommon/Private/ClothingAsset.cpp
+++ b/Engine/Source/Runtime/ClothingSystemRuntimeCommon/Private/ClothingAsset.cpp
@@ -1639,8 +1639,9 @@ void UClothingAssetCommon::PostEditChangeChainProperty(FPropertyChangedChainEven
 		const FName& PropertyName = ChainEvent.PropertyChain.GetActiveMemberNode() && ChainEvent.PropertyChain.GetActiveMemberNode()->GetNextNode() ?
 			ChainEvent.PropertyChain.GetActiveMemberNode()->GetNextNode()->GetValue()->GetFName() : NAME_None;
 
-		if (PropertyName == FName("SelfCollisionRadius") ||
-			PropertyName == FName("SelfCollisionCullScale"))
+		if (PropertyName == FName("bUseSelfCollisionSpheres") ||
+			PropertyName == FName("SelfCollisionSphereRadius") ||
+			PropertyName == FName("SelfCollisionSphereRadiusCullMultiplier"))
 		{
 			InvalidateFlaggedCachedData(EClothingCachedDataFlagsCommon::SelfCollisionData);
 			bReregisterComponents = true;
diff --git a/Engine/Source/Runtime/ClothingSystemRuntimeCommon/Public/ClothPhysicalMeshData.h b/Engine/Source/Runtime/ClothingSystemRuntimeCommon/Public/ClothPhysicalMeshData.h
--- a/Engine/Source/Runtime/ClothingSystemRuntimeCommon/Public/ClothPhysicalMeshData.h
+++ b/Engine/Source/Runtime/ClothingSystemRuntimeCommon/Public/ClothPhysicalMeshData.h
@@ -28,6 +28,11 @@ struct FClothPhysicalMeshData
 
 	/** Construct an empty cloth physical mesh with default common targets. */
 	CLOTHINGSYSTEMRUNTIMECOMMON_API FClothPhysicalMeshData();
+PRAGMA_DISABLE_DEPRECATION_WARNINGS
+	~FClothPhysicalMeshData() = default;
+	FClothPhysicalMeshData(const FClothPhysicalMeshData&) = default;
+	FClothPhysicalMeshData& operator=(const FClothPhysicalMeshData&) = default;
+PRAGMA_ENABLE_DEPRECATION_WARNINGS
 
 	/** Migrate from same, used to migrate LOD data from the UClothLODDataCommon_Legacy. */
 	CLOTHINGSYSTEMRUNTIMECOMMON_API void MigrateFrom(FClothPhysicalMeshData& ClothPhysicalMeshData);
@@ -122,7 +127,7 @@ struct FClothPhysicalMeshData
 
 	// Valid indices to use for self collisions (reduced set of Indices)
 	UPROPERTY(EditAnywhere, Category = SimMesh)
-	TArray<uint32> SelfCollisionIndices;
+	TSet<int32> SelfCollisionVertexSet;
 
 	// Long range attachment tethers, using euclidean (beeline) distance to find the closest attachment
 	UPROPERTY(EditAnywhere, Category = SimMesh)
@@ -140,6 +145,9 @@ struct FClothPhysicalMeshData
 	UPROPERTY(EditAnywhere, Category = SimMesh)
 	int32 NumFixedVerts;
 
+	UE_DEPRECATED(5.4, "Use SelfCollisionVertexSet instead.")
+	UPROPERTY()
+	TArray<uint32> SelfCollisionIndices;
 #if WITH_EDITORONLY_DATA
 	// Deprecated. Use WeightMaps instead.
 	UPROPERTY()
@@ -150,5 +158,6 @@ struct FClothPhysicalMeshData
 	TArray<float> BackstopRadiuses_DEPRECATED;
 	UPROPERTY()
 	TArray<float> AnimDriveMultipliers_DEPRECATED;
+
 #endif // WITH_EDITORONLY_DATA
 };
