Cloth - Added a progress bar to the deformer calculations.

diff --git a/Engine/Source/Runtime/ClothingSystemRuntimeCommon/Private/Utils/ClothingMeshUtils.cpp b/Engine/Source/Runtime/ClothingSystemRuntimeCommon/Private/Utils/ClothingMeshUtils.cpp
--- a/Engine/Source/Runtime/ClothingSystemRuntimeCommon/Private/Utils/ClothingMeshUtils.cpp
+++ b/Engine/Source/Runtime/ClothingSystemRuntimeCommon/Private/Utils/ClothingMeshUtils.cpp
@@ -7,6 +7,7 @@
 #include "Math/UnrealMathUtility.h"
 #include "Logging/LogMacros.h"
 #include "Async/ParallelFor.h"
+#include "Misc/ScopedSlowTask.h"
 
 #if WITH_EDITOR
 #include "Framework/Notifications/NotificationManager.h"
@@ -746,6 +747,11 @@ namespace ClothingMeshUtils
 		const TConstArrayView<float> MaxEdgeLengths = TargetMesh.GetMaxEdgeLengths();
 		check(NumMesh0Verts == MaxEdgeLengths.Num());
 
+		constexpr int32 SlowTaskDivider = 100;
+		const float NumSteps = (float)(NumMesh0Verts / SlowTaskDivider);
+		FScopedSlowTask SlowTask(NumSteps, LOCTEXT("GenerateMeshToMeshVertData", "Generating cloth deformer data..."));
+		SlowTask.MakeDialogDelayed(1.f);
+
 		if (bUseMultipleInfluences)
 		{
 			OutMeshToMeshVertData.Reserve(NumMesh0Verts * NUM_INFLUENCES_PER_VERTEX);
@@ -771,6 +777,11 @@ namespace ClothingMeshUtils
 				}
 
 				OutMeshToMeshVertData.Append(SkinningData.GetData(), NUM_INFLUENCES_PER_VERTEX);
+
+				if ((VertIdx0 + 1) % SlowTaskDivider == 0)
+				{
+					SlowTask.EnterProgressFrame();
+				}
 			}
 
 			check(OutMeshToMeshVertData.Num() == NumMesh0Verts * NUM_INFLUENCES_PER_VERTEX);
@@ -804,8 +815,12 @@ namespace ClothingMeshUtils
 				check(ClosestTriangleBaseIdx != INDEX_NONE);
 
 				SingleSkinningDataForVertex(VertPosition, VertNormal, (FVector3f)VertTangent, SourceMesh, ClosestTriangleBaseIdx, SkinningData);
-			}
 
+				if ((VertIdx0 + 1) % SlowTaskDivider == 0)
+				{
+					SlowTask.EnterProgressFrame();
+				}
+			}
 			check(OutMeshToMeshVertData.Num() == NumMesh0Verts);
 		}
 
