ClothingMeshUtils: if we encounter degenerate triangles during skinning, only do a Slate notification once rather than for each vertex. Multiple slate pop-ups can cause the Editor to become temporarily unusable.

diff --git a/Engine/Source/Runtime/ClothingSystemRuntimeCommon/Private/Utils/ClothingMeshUtils.cpp b/Engine/Source/Runtime/ClothingSystemRuntimeCommon/Private/Utils/ClothingMeshUtils.cpp
--- a/Engine/Source/Runtime/ClothingSystemRuntimeCommon/Private/Utils/ClothingMeshUtils.cpp
+++ b/Engine/Source/Runtime/ClothingSystemRuntimeCommon/Private/Utils/ClothingMeshUtils.cpp
@@ -493,6 +493,7 @@ namespace ClothingMeshUtils
 			return OneMinusR2 * OneMinusR2 * OneMinusR2;
 		}
 
+		// Return false if triangle specified by ClosestTriangleBaseIdx is degenerate, true otherwise
 		bool SingleSkinningDataForVertex(const FVector3f& VertPosition,
 										 const FVector3f& VertNormal,
 										 const FVector3f& VertTangent,
@@ -517,16 +518,15 @@ namespace ClothingMeshUtils
 			{
 				// Failed, we have 2 identical vertices
 
-				// Log and toast
-				FText Error = FText::Format(LOCTEXT("DegenerateTriangleError", "Failed to generate skinning data, found conincident vertices in triangle A={0} B={1} C={2}"), FText::FromString(A.ToString()), FText::FromString(B.ToString()), FText::FromString(C.ToString()));
-
+				// Log
+				const uint32 IndexA = SourceMesh.GetIndices()[ClosestTriangleBaseIdx];
+				const uint32 IndexB = SourceMesh.GetIndices()[ClosestTriangleBaseIdx + 1];
+				const uint32 IndexC = SourceMesh.GetIndices()[ClosestTriangleBaseIdx + 2];
+				FText Error = FText::Format(LOCTEXT("DegenerateTriangleError", "Failed to generate skinning data, found coincident vertices in triangle ({0}, {1}, {2}), points A={3} B={4} C={5}"),
+					IndexA, IndexB, IndexC,
+					FText::FromString(A.ToString()), FText::FromString(B.ToString()), FText::FromString(C.ToString()));
 				UE_LOG(LogClothingMeshUtils, Warning, TEXT("%s"), *Error.ToString());
 
-#if WITH_EDITOR
-				FNotificationInfo Info(Error);
-				Info.ExpireDuration = 5.0f;
-				FSlateNotificationManager::Get().AddNotification(Info);
-#endif
 				return false;
 			}
 
@@ -602,14 +602,7 @@ namespace ClothingMeshUtils
 
 					// Log and toast
 					FText Error = FText::Format(LOCTEXT("DegenerateTriangleError", "Failed to generate skinning data, found conincident vertices in triangle A={0} B={1} C={2}"), FText::FromString(A.ToString()), FText::FromString(B.ToString()), FText::FromString(C.ToString()));
-
 					UE_LOG(LogClothingMeshUtils, Warning, TEXT("%s"), *Error.ToString());
-
-#if WITH_EDITOR
-					FNotificationInfo Info(Error);
-					Info.ExpireDuration = 5.0f;
-					FSlateNotificationManager::Get().AddNotification(Info);
-#endif
 					return false;
 				}
 
@@ -654,6 +647,8 @@ namespace ClothingMeshUtils
 			const TConstArrayView<float> MaxEdgeLengths = TargetMesh.GetMaxEdgeLengths();
 			check(NumTargetMeshVerts == MaxEdgeLengths.Num());
 
+			bool bAnySmallTriangleEncountered = false;
+
 			for (int32 VID = 0; VID < NumTargetMeshVerts; ++VID)
 			{
 				float SumWeight = 0.0f;
@@ -695,8 +690,13 @@ namespace ClothingMeshUtils
 					else
 					{
 						const uint16 PreviousFlag = FirstData.SourceMeshVertIndices[3];
-						SingleSkinningDataForVertex(VertPosition, VertNormal, VertTangent, SourceMesh, ClosestTriangleBaseIdx, FirstData);
+						const bool bSkinningSuccess = SingleSkinningDataForVertex(VertPosition, VertNormal, VertTangent, SourceMesh, ClosestTriangleBaseIdx, FirstData);
 						FirstData.SourceMeshVertIndices[3] = PreviousFlag;
+
+						if (!bSkinningSuccess)
+						{
+							bAnySmallTriangleEncountered = true;
+						}
 					}
 
 					// Set all other skinning data to have zero weight
@@ -708,6 +708,17 @@ namespace ClothingMeshUtils
 					}
 				}
 			}
+
+#if WITH_EDITOR
+			if (bAnySmallTriangleEncountered)
+			{
+				const FText ErrorMsg = LOCTEXT("DegenerateTriangleErrorToast", "Failed to generate skinning data, found conincident vertices in at least one triangle. See Log for details");
+				FNotificationInfo Info(ErrorMsg);
+				Info.ExpireDuration = 5.0f;
+				FSlateNotificationManager::Get().AddNotification(Info);
+			}
+#endif
+
 		}
 
 	}  // Anonymous namespace
@@ -791,6 +802,7 @@ namespace ClothingMeshUtils
 			OutMeshToMeshVertData.Reserve(NumMesh0Verts);
 
 			// For all mesh0 verts
+			bool bAnySmallTriangleEncountered = false;
 			for (int32 VertIdx0 = 0; VertIdx0 < NumMesh0Verts; ++VertIdx0)
 			{
 				OutMeshToMeshVertData.AddZeroed();
@@ -814,14 +826,31 @@ namespace ClothingMeshUtils
 				const int32 ClosestTriangleBaseIdx = GetBestTriangleBaseIndex(SourceMesh, (FVector)VertPosition, MaxEdgeLengths[VertIdx0]);
 				check(ClosestTriangleBaseIdx != INDEX_NONE);
 
-				SingleSkinningDataForVertex(VertPosition, VertNormal, (FVector3f)VertTangent, SourceMesh, ClosestTriangleBaseIdx, SkinningData);
+				const bool bSkinningSuccess = SingleSkinningDataForVertex(VertPosition, VertNormal, (FVector3f)VertTangent, SourceMesh, ClosestTriangleBaseIdx, SkinningData);
+				
+				if (!bSkinningSuccess)
+				{
+					bAnySmallTriangleEncountered = true;
+				}
 
 				if ((VertIdx0 + 1) % SlowTaskDivider == 0)
 				{
 					SlowTask.EnterProgressFrame();
 				}
 			}
+
 			check(OutMeshToMeshVertData.Num() == NumMesh0Verts);
+
+#if WITH_EDITOR
+			if (bAnySmallTriangleEncountered)
+			{
+				const FText ErrorMsg = LOCTEXT("DegenerateTriangleErrorToast", "Failed to generate skinning data, found conincident vertices in at least one triangle. See Log for details");
+				FNotificationInfo Info(ErrorMsg);
+				Info.ExpireDuration = 5.0f;
+				FSlateNotificationManager::Get().AddNotification(Info);
+			}
+#endif
+
 		}
 
 		if (OutMeshToMeshVertData.Num())
@@ -862,6 +891,8 @@ namespace ClothingMeshUtils
 			return;
 		}
 
+		bool bAnySmallTriangleEncountered = false;
+
 		for (int32 VID = 0; VID < NumTargetMeshVerts; ++VID)
 		{
 			float SumWeight = 0.0f;
@@ -903,8 +934,13 @@ namespace ClothingMeshUtils
 				else
 				{
 					const uint16 PreviousFlag = FirstData.SourceMeshVertIndices[3];
-					SingleSkinningDataForVertex(VertPosition, VertNormal, VertTangent, SourceMesh, ClosestTriangleBaseIdx, FirstData);
+					const bool bSkinningSuccess = SingleSkinningDataForVertex(VertPosition, VertNormal, VertTangent, SourceMesh, ClosestTriangleBaseIdx, FirstData);
 					FirstData.SourceMeshVertIndices[3] = PreviousFlag;
+
+					if (!bSkinningSuccess)
+					{
+						bAnySmallTriangleEncountered = true;
+					}
 				}
 
 				// Set all other skinning data to have zero weight
@@ -916,6 +952,17 @@ namespace ClothingMeshUtils
 				}
 			}
 		}
+
+#if WITH_EDITOR
+		if (bAnySmallTriangleEncountered)
+		{
+			const FText ErrorMsg = LOCTEXT("DegenerateTriangleErrorToast", "Failed to generate skinning data, found conincident vertices in at least one triangle. See Log for details");
+			FNotificationInfo Info(ErrorMsg);
+			Info.ExpireDuration = 5.0f;
+			FSlateNotificationManager::Get().AddNotification(Info);
+		}
+#endif
+
 	}
 
 	// TODO: Vertex normals are not used at present, a future improved algorithm might however
