ClothingMeshUtils: use the same error format for degenerate triangles in both the single-influence and multiple-influence code path

diff --git a/Engine/Source/Runtime/ClothingSystemRuntimeCommon/Private/Utils/ClothingMeshUtils.cpp b/Engine/Source/Runtime/ClothingSystemRuntimeCommon/Private/Utils/ClothingMeshUtils.cpp
--- a/Engine/Source/Runtime/ClothingSystemRuntimeCommon/Private/Utils/ClothingMeshUtils.cpp
+++ b/Engine/Source/Runtime/ClothingSystemRuntimeCommon/Private/Utils/ClothingMeshUtils.cpp
@@ -600,9 +600,14 @@ namespace ClothingMeshUtils
 				{
 					// Failed, we have 2 identical vertices
 
-					// Log and toast
-					FText Error = FText::Format(LOCTEXT("DegenerateTriangleError", "Failed to generate skinning data, found conincident vertices in triangle A={0} B={1} C={2}"), FText::FromString(A.ToString()), FText::FromString(B.ToString()), FText::FromString(C.ToString()));
+					const uint32 IndexA = SourceMesh.GetIndices()[ClosestTriangleBaseIdx];
+					const uint32 IndexB = SourceMesh.GetIndices()[ClosestTriangleBaseIdx + 1];
+					const uint32 IndexC = SourceMesh.GetIndices()[ClosestTriangleBaseIdx + 2];
+					FText Error = FText::Format(LOCTEXT("DegenerateTriangleErrorMultipleInfluences", "Failed to generate skinning data, found coincident vertices in triangle ({0}, {1}, {2}), points A={3} B={4} C={5}"),
+						IndexA, IndexB, IndexC,
+						FText::FromString(A.ToString()), FText::FromString(B.ToString()), FText::FromString(C.ToString()));
 					UE_LOG(LogClothingMeshUtils, Warning, TEXT("%s"), *Error.ToString());
+
 					return false;
 				}
 
