Cloth mesh-to-mesh skinning: check that the area of a normal-offset triangle is non-zero before trying to compute barycentric coordinates for it

diff --git a/Engine/Source/Runtime/ClothingSystemRuntimeCommon/Private/Utils/ClothingMeshUtils.cpp b/Engine/Source/Runtime/ClothingSystemRuntimeCommon/Private/Utils/ClothingMeshUtils.cpp
--- a/Engine/Source/Runtime/ClothingSystemRuntimeCommon/Private/Utils/ClothingMeshUtils.cpp
+++ b/Engine/Source/Runtime/ClothingSystemRuntimeCommon/Private/Utils/ClothingMeshUtils.cpp
@@ -1225,13 +1225,19 @@ namespace ClothingMeshUtils
 			FPlane4f TrianglePlane(AW, BW, CW);
 
 			const FVector3f PointOnTriPlane = FVector3f::PointPlaneProject(Point, TrianglePlane);
-			const FVector3f BaryCoords = (FVector3f)FMath::ComputeBaryCentric2D((FVector)PointOnTriPlane, (FVector)AW, (FVector)BW, (FVector)CW);
-
-			if (BaryCoords.X == BaryCoords.Y && BaryCoords.Y == BaryCoords.Z && BaryCoords.Z == 0.0f)
+			
+			// check the triangle {A+wNA, B+wNB, C+wNC} is not degenerate
+			const FVector TriNorm = ((FVector)BW - (FVector)AW).Cross((FVector)CW - (FVector)AW);
+			const double TriNormSizeSquared = TriNorm.SizeSquared();
+			if (TriNormSizeSquared <= UE_DOUBLE_SMALL_NUMBER)
 			{
-				// Degenerate triangle at this value of W
 				continue;
 			}
+			 
+			const FVector3f BaryCoords = (FVector3f)FMath::ComputeBaryCentric2D((FVector)PointOnTriPlane, (FVector)AW, (FVector)BW, (FVector)CW);
+
+			checkf(!(BaryCoords.X == BaryCoords.Y && BaryCoords.Y == BaryCoords.Z && BaryCoords.Z == 0.0f), 
+				TEXT("ComputeBaryCentric2D returned all zeros despite triangle area being non-zero"));
 
 			bAnySolutionFound = true;
 
