Cloth Asset - Added selection restriction sets for cloth deformer mappings.

#rnx
#tests Created a filter dual selection for the deformer, and that fixed the issue on the production asset I tested.

diff --git a/Engine/Plugins/Experimental/ChaosClothAssetEditor/Source/ChaosClothAssetDataflowNodes/Private/ChaosClothAsset/ChaosClothAssetDataflowNodesModule.cpp b/Engine/Plugins/Experimental/ChaosClothAssetEditor/Source/ChaosClothAssetDataflowNodes/Private/ChaosClothAsset/ChaosClothAssetDataflowNodesModule.cpp
--- a/Engine/Plugins/Experimental/ChaosClothAssetEditor/Source/ChaosClothAssetDataflowNodes/Private/ChaosClothAsset/ChaosClothAssetDataflowNodesModule.cpp
+++ b/Engine/Plugins/Experimental/ChaosClothAssetEditor/Source/ChaosClothAssetDataflowNodes/Private/ChaosClothAsset/ChaosClothAssetDataflowNodesModule.cpp
@@ -136,6 +136,7 @@ namespace UE::Chaos::ClothAsset
 				PropertyModule->RegisterCustomPropertyTypeLayout(FChaosClothAssetWeightedValue::StaticStruct()->GetFName(), FOnGetPropertyTypeCustomizationInstance::CreateStatic(&FWeightedValueCustomization::MakeInstance));
 				PropertyModule->RegisterCustomPropertyTypeLayout(FChaosClothAssetWeightedValueNonAnimatable::StaticStruct()->GetFName(), FOnGetPropertyTypeCustomizationInstance::CreateStatic(&FWeightedValueCustomization::MakeInstance));
 				PropertyModule->RegisterCustomPropertyTypeLayout(FChaosClothAssetWeightedValueNonAnimatableNoLowHighRange::StaticStruct()->GetFName(), FOnGetPropertyTypeCustomizationInstance::CreateStatic(&FWeightedValueCustomization::MakeInstance));
+				PropertyModule->RegisterCustomPropertyTypeLayout(FChaosClothAssetConnectableStringValue::StaticStruct()->GetFName(), FOnGetPropertyTypeCustomizationInstance::CreateStatic(&FWeightedValueCustomization::MakeInstance));
 				PropertyModule->RegisterCustomPropertyTypeLayout(FChaosClothAssetConnectableIStringValue::StaticStruct()->GetFName(), FOnGetPropertyTypeCustomizationInstance::CreateStatic(&FWeightedValueCustomization::MakeInstance));
 				PropertyModule->RegisterCustomPropertyTypeLayout(FChaosClothAssetConnectableIOStringValue::StaticStruct()->GetFName(), FOnGetPropertyTypeCustomizationInstance::CreateStatic(&FWeightedValueCustomization::MakeInstance));
 				PropertyModule->RegisterCustomPropertyTypeLayout(FChaosClothAssetNodeSelectionGroup::StaticStruct()->GetFName(), FOnGetPropertyTypeCustomizationInstance::CreateStatic(&FSelectionGroupCustomization::MakeInstance));
diff --git a/Engine/Plugins/Experimental/ChaosClothAssetEditor/Source/ChaosClothAssetDataflowNodes/Private/ChaosClothAsset/ProxyDeformerNode.cpp b/Engine/Plugins/Experimental/ChaosClothAssetEditor/Source/ChaosClothAssetDataflowNodes/Private/ChaosClothAsset/ProxyDeformerNode.cpp
--- a/Engine/Plugins/Experimental/ChaosClothAssetEditor/Source/ChaosClothAssetDataflowNodes/Private/ChaosClothAsset/ProxyDeformerNode.cpp
+++ b/Engine/Plugins/Experimental/ChaosClothAssetEditor/Source/ChaosClothAssetDataflowNodes/Private/ChaosClothAsset/ProxyDeformerNode.cpp
@@ -1,9 +1,11 @@
 // Copyright Epic Games, Inc. All Rights Reserved.
 
 #include "ChaosClothAsset/ProxyDeformerNode.h"
+#include "ChaosClothAsset/ClothCollectionAttribute.h"
+#include "ChaosClothAsset/ClothCollectionGroup.h"
 #include "ChaosClothAsset/ClothDataflowTools.h"
 #include "ChaosClothAsset/CollectionClothFacade.h"
-#include "ChaosClothAsset/ClothCollectionAttribute.h"
+#include "ChaosClothAsset/CollectionClothSelectionFacade.h"
 #include "Dataflow/DataflowInputOutput.h"
 #include "Utils/ClothingMeshUtils.h"
 #include "PointWeightMap.h"
@@ -18,10 +20,12 @@ namespace UE::Chaos::ClothAsset::Private
 	{
 		TConstArrayView<FVector3f> SimPositions;
 		TConstArrayView<FIntVector3> SimIndices;
-		TConstArrayView<float> MaxDistances;
 		TConstArrayView<FVector3f> RenderPositions;
 		TConstArrayView<FVector3f> RenderNormals;
 		TConstArrayView<FIntVector3> RenderIndices;
+		FPointWeightMap PointWeightMap;
+		TArray<ClothingMeshUtils::FMeshToMeshFilterSet> MeshToMeshFilterSet;
+
 		TArrayView<TArray<FVector4f>> RenderDeformerPositionBaryCoordsAndDist;
 		TArrayView<TArray<FVector4f>> RenderDeformerNormalBaryCoordsAndDist;
 		TArrayView<TArray<FVector4f>> RenderDeformerTangentBaryCoordsAndDist;
@@ -64,8 +68,6 @@ namespace UE::Chaos::ClothAsset::Private
 			const ClothingMeshUtils::ClothMeshDesc SimMeshDesc(SimPositions, ScalarSimIndices);
 			const ClothingMeshUtils::ClothMeshDesc RenderMeshDesc(RenderPositions, RenderNormals, ScalarRenderIndices);
 
-			FPointWeightMap PointWeightMap = (MaxDistances.Num() == SimPositions.Num()) ? FPointWeightMap(MaxDistances) : FPointWeightMap(SimPositions.Num(), 1.f);
-
 			TArray<FMeshToMeshVertData> MeshToMeshVertData;
 
 			ClothingMeshUtils::GenerateMeshToMeshVertData(
@@ -75,7 +77,8 @@ namespace UE::Chaos::ClothAsset::Private
 				&PointWeightMap,
 				bUseSmoothTransition,
 				bUseMultipleInfluences,
-				InfluenceRadius);
+				InfluenceRadius,
+				MeshToMeshFilterSet);
 
 			const int32 NumInfluences = MeshToMeshVertData.Num() / RenderPositions.Num();
 			check(MeshToMeshVertData.Num() == RenderPositions.Num() * NumInfluences);  // Check modulo
@@ -110,17 +113,183 @@ namespace UE::Chaos::ClothAsset::Private
 			return NumInfluences;
 		}
 	};
-}
+
+	static FPointWeightMap SelectionToPointWeightMap(const FCollectionClothConstFacade& ClothFacade, const FCollectionClothSelectionConstFacade& SelectionFacade, const FName& SelectionName)
+	{
+		constexpr float SelectedValue = 1.f;
+		if (const TSet<int32>* SelectionSet = SelectionFacade.IsValid() ? SelectionFacade.FindSelectionSet(SelectionName) : nullptr)
+		{
+			FPointWeightMap PointWeightMap;
+
+			const FName SelectionGroup = SelectionFacade.GetSelectionGroup(SelectionName);
+
+			if (SelectionGroup == ClothCollectionGroup::SimVertices3D)
+			{
+				PointWeightMap.Initialize(ClothFacade.GetNumSimVertices3D());  // Init to zero (unselected)
+				for (const int32 VertexIndex : *SelectionSet)
+				{
+					PointWeightMap[VertexIndex] = SelectedValue;
+				}
+				return PointWeightMap;
+			}
+			else if (SelectionGroup == ClothCollectionGroup::SimVertices2D)
+			{
+				PointWeightMap.Initialize(ClothFacade.GetNumSimVertices3D());  // Init to zero (unselected)
+				const TConstArrayView<int32> Vertex2DTo3D = ClothFacade.GetSimVertex3DLookup();
+				for (const int32 VertexIndex : *SelectionSet)
+				{
+					PointWeightMap[Vertex2DTo3D[VertexIndex]] = SelectedValue;
+				}
+				return PointWeightMap;
+			}
+			else if (SelectionGroup == ClothCollectionGroup::SimFaces)
+			{
+				PointWeightMap.Initialize(ClothFacade.GetNumSimVertices3D());  // Init to zero (unselected)
+				const TConstArrayView<FIntVector3> SimIndices3D = ClothFacade.GetSimIndices3D();
+				for (const int32 FaceIndex : *SelectionSet)
+				{
+					PointWeightMap[SimIndices3D[FaceIndex][0]] = SelectedValue;
+					PointWeightMap[SimIndices3D[FaceIndex][1]] = SelectedValue;
+					PointWeightMap[SimIndices3D[FaceIndex][2]] = SelectedValue;
+				}
+				return PointWeightMap;
+			}
+		}
+		// Invalid or no selection found, all points are dynamic 
+		return FPointWeightMap(ClothFacade.GetNumSimVertices3D(), SelectedValue);
+	}
+
+	static TArray<ClothingMeshUtils::FMeshToMeshFilterSet> SelectionsToMeshToMeshFilterSets(const FCollectionClothConstFacade& ClothFacade, const FCollectionClothSelectionConstFacade& SelectionFacade, const TArray<FName>& SelectionNames)
+	{
+		auto GetSimFaceSelection = [&ClothFacade](const FName& SelectionGroup, const TSet<int32>& SelectionSet) -> TSet<int32>
+			{
+				TSet<int32> SimFaceSelection;
+				if (SelectionGroup == ClothCollectionGroup::SimVertices2D)
+				{
+					const TConstArrayView<FIntVector3> SimIndices2D = ClothFacade.GetSimIndices2D();
+					SimFaceSelection.Reserve(SelectionSet.Num());
+
+					for (int32 FaceIndex = 0; FaceIndex < SimIndices2D.Num(); ++FaceIndex)
+					{
+						const FIntVector3& Indices = SimIndices2D[FaceIndex];
+
+						if (SelectionSet.Contains(Indices[0]) &&
+							SelectionSet.Contains(Indices[1]) &&
+							SelectionSet.Contains(Indices[2]))
+						{
+							SimFaceSelection.Add(FaceIndex);
+						}
+					}
+				}
+				else if (SelectionGroup == ClothCollectionGroup::SimVertices3D)
+				{
+					const TConstArrayView<FIntVector3> SimIndices3D = ClothFacade.GetSimIndices3D();
+					SimFaceSelection.Reserve(SelectionSet.Num());
+
+					for (int32 FaceIndex = 0; FaceIndex < SimIndices3D.Num(); ++FaceIndex)
+					{
+						const FIntVector3& Indices = SimIndices3D[FaceIndex];
+
+						if (SelectionSet.Contains(Indices[0]) &&
+							SelectionSet.Contains(Indices[1]) &&
+							SelectionSet.Contains(Indices[2]))
+						{
+							SimFaceSelection.Add(FaceIndex);
+						}
+					}
+				}
+				else if (SelectionGroup == ClothCollectionGroup::SimFaces)
+				{
+					SimFaceSelection = SelectionSet;
+				}
+				return SimFaceSelection;
+			};
+
+		auto GetRenderVertexSelection = [&ClothFacade](const FName& SelectionGroup, const TSet<int32>& SelectionSet) -> TSet<int32>
+			{
+				TSet<int32> RenderVertexSelection;
+				if (SelectionGroup == ClothCollectionGroup::RenderVertices)
+				{
+					RenderVertexSelection = SelectionSet;
+				}
+				else if (SelectionGroup == ClothCollectionGroup::RenderFaces)
+				{
+					RenderVertexSelection.Reserve(SelectionSet.Num() * 3);
+					const TConstArrayView<FIntVector3> RenderIndices = ClothFacade.GetRenderIndices();
+					for (const int32 FaceIndex : SelectionSet)
+					{
+						RenderVertexSelection.Add(RenderIndices[FaceIndex][0]);
+						RenderVertexSelection.Add(RenderIndices[FaceIndex][1]);
+						RenderVertexSelection.Add(RenderIndices[FaceIndex][2]);
+					}
+				}
+				return RenderVertexSelection;
+			};
+
+		// Fill up the MeshToMeshFilterSets
+		TArray<ClothingMeshUtils::FMeshToMeshFilterSet> MeshToMeshFilterSets;
+
+		if (SelectionFacade.IsValid())
+		{
+			MeshToMeshFilterSets.Reserve(SelectionNames.Num());
+
+			for (const FName& SelectionName : SelectionNames)
+			{
+				if (const TSet<int32>* SelectionSet = SelectionFacade.FindSelectionSet(SelectionName))
+				{
+					if (const TSet<int32>* SecondarySelectionSet = SelectionFacade.FindSelectionSecondarySet(SelectionName))
+					{
+						if (SelectionSet->Num() && SecondarySelectionSet->Num())
+						{
+							FName SelectionGroup = SelectionFacade.GetSelectionGroup(SelectionName);
+							FName SelectionSecondaryGroup = SelectionFacade.GetSelectionSecondaryGroup(SelectionName);
+
+							// Retrieve the sim face selection
+							TSet<int32> SimFaceSelection = GetSimFaceSelection(SelectionGroup, *SelectionSet);
+							if (!SimFaceSelection.Num())
+							{
+								// Try swapping the selections
+								Swap(SelectionSet, SecondarySelectionSet);
+								Swap(SelectionGroup, SelectionSecondaryGroup);
+
+								SimFaceSelection = GetSimFaceSelection(SelectionGroup, *SelectionSet);
+							}
+							if (!SimFaceSelection.Num())
+							{
+								continue;  // Nothing selected on the simulation side
+							}
+
+							// Retrieve the render vertex selection
+							TSet<int32> RenderVertexSelection = GetRenderVertexSelection(SelectionSecondaryGroup, *SecondarySelectionSet);
+							if (!RenderVertexSelection.Num())
+							{
+								continue;  // Nothing selected on the render side
+							}
+
+							ClothingMeshUtils::FMeshToMeshFilterSet& MeshToMeshFilterSet = MeshToMeshFilterSets.AddDefaulted_GetRef();
+							MeshToMeshFilterSet.SourceTriangles = MoveTemp(SimFaceSelection);
+							MeshToMeshFilterSet.TargetVertices = MoveTemp(RenderVertexSelection);
+						}
+					}
+				}
+			}
+		}
+
+		return MeshToMeshFilterSets;
+	}
+
+}  // End namespace UE::Chaos::ClothAsset::Private
 
 FChaosClothAssetProxyDeformerNode::FChaosClothAssetProxyDeformerNode(const Dataflow::FNodeParameters& InParam, FGuid InGuid)
 	: FDataflowNode(InParam, InGuid)
 {
 	using namespace UE::Chaos::ClothAsset;
-	MaxDistance.WeightMap = FString();  // An empty weight map is an accepted input, but a non existing one isn't
+	SimVertexSelection.StringValue = FString();  // An empty selection is an accepted input, but a non existing one isn't
 	SkinningBlendName = ClothCollectionAttribute::RenderDeformerSkinningBlend.ToString();
 
 	RegisterInputConnection(&Collection);
-	RegisterInputConnection(&MaxDistance.WeightMap, GET_MEMBER_NAME_CHECKED(FChaosClothAssetWeightedValueNonAnimatableNoLowHighRange, WeightMap));
+	RegisterInputConnection(&SimVertexSelection.StringValue, GET_MEMBER_NAME_CHECKED(FChaosClothAssetConnectableIStringValue, StringValue));
+	RegisterInputConnection(&SelectionFilterSet0.StringValue, GET_MEMBER_NAME_CHECKED(FChaosClothAssetConnectableIStringValue, StringValue));
 	RegisterOutputConnection(&Collection, &Collection);
 	RegisterOutputConnection(&SkinningBlendName);
 }
@@ -131,8 +300,15 @@ void FChaosClothAssetProxyDeformerNode::Evaluate(Dataflow::FContext& Context, co
 	{
 		using namespace UE::Chaos::ClothAsset;
 
-		// Update the weight map override
-		MaxDistance.WeightMap_Override = GetValue<FString>(Context, &MaxDistance.WeightMap, UE::Chaos::ClothAsset::FWeightMapTools::NotOverridden);
+		// Update the selection names override
+		SimVertexSelection.StringValue_Override = GetValue<FString>(Context, &SimVertexSelection.StringValue, UE::Chaos::ClothAsset::FWeightMapTools::NotOverridden);
+		SelectionFilterSet0.StringValue_Override = GetValue<FString>(Context, &SelectionFilterSet0.StringValue, UE::Chaos::ClothAsset::FWeightMapTools::NotOverridden);
+		const TArray<const FChaosClothAssetConnectableStringValue*> Non0SelectionFilterSets = Get1To9SelectionFilterSets();
+		for (int32 FilterSetIndex = 1; FilterSetIndex < NumFilterSets; ++FilterSetIndex)
+		{
+			const FChaosClothAssetConnectableStringValue& SelectionFilterSet = *Non0SelectionFilterSets[FilterSetIndex - 1];
+			SelectionFilterSet.StringValue_Override = GetValue<FString>(Context, &SelectionFilterSet.StringValue, UE::Chaos::ClothAsset::FWeightMapTools::NotOverridden);
+		}
 
 		// Evaluate in collection
 		FManagedArrayCollection InCollection = GetValue<FManagedArrayCollection>(Context, &Collection);
@@ -142,15 +318,16 @@ void FChaosClothAssetProxyDeformerNode::Evaluate(Dataflow::FContext& Context, co
 		FCollectionClothFacade ClothFacade(ClothCollection);
 		if (ClothFacade.IsValid() && ClothFacade.HasValidData())
 		{
-			// Retrieve the MaxDistance weight map name
-			FName MaxDistanceWeightMapName = FName(*GetValue<FString>(Context, &MaxDistance.WeightMap));
-			if (MaxDistanceWeightMapName != NAME_None && !ClothFacade.HasWeightMap(MaxDistanceWeightMapName))
+			FCollectionClothSelectionFacade SelectionFacade(ClothCollection);
+
+			// Retrieve the SimVertexSelection name
+			FName SimVertexSelectionName = FName(*GetValue<FString>(Context, &SimVertexSelection.StringValue));
+			if (SimVertexSelectionName != NAME_None && (!SelectionFacade.IsValid() || !SelectionFacade.FindSelectionSet(SimVertexSelectionName)))
 			{
 				FClothDataflowTools::LogAndToastWarning(*this,
-					LOCTEXT("ClothFacadeHasWeightMapHeadline", "Unknown MaxDistance weight map."),
-					LOCTEXT("ClothFacadeHasWeightMapDetails", "The specified MaxDistance weight map does't exist within the input Cloth Collection."));
-
-				MaxDistanceWeightMapName = NAME_None;
+					LOCTEXT("HasSimVertexSelectionHeadline", "Unknown SimVertexSelection."),
+					LOCTEXT("HasSimVertexSelectionDetails", "The specified SimVertexSelection does't exist within the input Cloth Collection."));
+				SimVertexSelectionName = NAME_None;
 			}
 
 			// Add the optional render deformer schema
@@ -163,10 +340,11 @@ void FChaosClothAssetProxyDeformerNode::Evaluate(Dataflow::FContext& Context, co
 			Private::FDeformerMappingDataGenerator DeformerMappingDataGenerator;
 			DeformerMappingDataGenerator.SimPositions = ClothFacade.GetSimPosition3D();
 			DeformerMappingDataGenerator.SimIndices = ClothFacade.GetSimIndices3D();
-			DeformerMappingDataGenerator.MaxDistances = ClothFacade.GetWeightMap(MaxDistanceWeightMapName);
 			DeformerMappingDataGenerator.RenderPositions = ClothFacade.GetRenderPosition();
 			DeformerMappingDataGenerator.RenderNormals = ClothFacade.GetRenderNormal();
 			DeformerMappingDataGenerator.RenderIndices = ClothFacade.GetRenderIndices();
+			DeformerMappingDataGenerator.PointWeightMap = Private::SelectionToPointWeightMap(ClothFacade, SelectionFacade, SimVertexSelectionName);
+			DeformerMappingDataGenerator.MeshToMeshFilterSet = Private::SelectionsToMeshToMeshFilterSets(ClothFacade, SelectionFacade, GetSelectionFilterNames(Context));
 			DeformerMappingDataGenerator.RenderDeformerPositionBaryCoordsAndDist = ClothFacade.GetRenderDeformerPositionBaryCoordsAndDist();
 			DeformerMappingDataGenerator.RenderDeformerNormalBaryCoordsAndDist = ClothFacade.GetRenderDeformerNormalBaryCoordsAndDist();
 			DeformerMappingDataGenerator.RenderDeformerTangentBaryCoordsAndDist = ClothFacade.GetRenderDeformerTangentBaryCoordsAndDist();
@@ -187,4 +365,89 @@ void FChaosClothAssetProxyDeformerNode::Evaluate(Dataflow::FContext& Context, co
 	}
 }
 
+Dataflow::FPin FChaosClothAssetProxyDeformerNode::AddPin()
+{
+	check(NumFilterSets > 0);
+	const FChaosClothAssetConnectableStringValue* const SelectionFilterSet = Get1To9SelectionFilterSets()[NumFilterSets - 1];
+
+	RegisterInputConnection(&SelectionFilterSet->StringValue, GET_MEMBER_NAME_CHECKED(FChaosClothAssetConnectableIStringValue, StringValue));
+	++NumFilterSets;
+	const FDataflowInput* const Input = FindInput(SelectionFilterSet);
+	check(Input);
+	return { Dataflow::FPin::EDirection::INPUT, Input->GetType(), Input->GetName() };
+}
+
+Dataflow::FPin FChaosClothAssetProxyDeformerNode::GetPinToRemove() const
+{
+	check(NumFilterSets > 1);
+	const FChaosClothAssetConnectableStringValue* const SelectionFilterSet = Get1To9SelectionFilterSets()[NumFilterSets - 2];
+	const FDataflowInput* const Input = FindInput(SelectionFilterSet);
+	check(Input);
+	return { Dataflow::FPin::EDirection::INPUT, Input->GetType(), Input->GetName() };
+}
+
+void FChaosClothAssetProxyDeformerNode::OnPinRemoved(const Dataflow::FPin& Pin)
+{
+	check(NumFilterSets > 1);
+	const FChaosClothAssetConnectableStringValue* const SelectionFilterSet = Get1To9SelectionFilterSets()[NumFilterSets - 2];
+	check(Pin.Direction == Dataflow::FPin::EDirection::INPUT);
+#if DO_CHECK
+	const FDataflowInput* const Input = FindInput(SelectionFilterSet);
+	check(Input);
+	check(Input->GetName() == Pin.Name);
+	check(Input->GetType() == Pin.Type);
+#endif
+	--NumFilterSets;
+	return Super::OnPinRemoved(Pin);
+}
+
+void FChaosClothAssetProxyDeformerNode::Serialize(FArchive& Ar)
+{
+	// Restore the pins when re-loading so they can get properly reconnected
+	if (Ar.IsLoading())
+	{
+		const int32 NumFilterSetsToAdd = (NumFilterSets - 1);
+		NumFilterSets = 1;  // Reset to default, add pin will increment it again 
+		for (int32 Index = 0; Index < NumFilterSetsToAdd; ++Index)
+		{
+			AddPin();
+		}
+		ensure(NumFilterSetsToAdd == (NumFilterSets - 1));
+	}
+}
+
+TArray<FName> FChaosClothAssetProxyDeformerNode::GetSelectionFilterNames(Dataflow::FContext& Context) const
+{
+	check(NumFilterSets > 0);
+
+	TArray<FName> SelectionFilterSets;
+	SelectionFilterSets.SetNumUninitialized(NumFilterSets);
+
+	SelectionFilterSets[0] = FName(*GetValue(Context, &SelectionFilterSet0.StringValue));
+
+	TArray<const FChaosClothAssetConnectableStringValue*> Non0SelectionFilterSets = Get1To9SelectionFilterSets();
+
+	for (int32 FilterSetIndex = 1; FilterSetIndex < NumFilterSets; ++FilterSetIndex)
+	{
+		SelectionFilterSets[FilterSetIndex] = FName(*GetValue(Context, &Non0SelectionFilterSets[FilterSetIndex]->StringValue));
+	}
+	return SelectionFilterSets;
+}
+
+TArray<const FChaosClothAssetConnectableStringValue*> FChaosClothAssetProxyDeformerNode::Get1To9SelectionFilterSets() const
+{
+	return TArray<const FChaosClothAssetConnectableStringValue*>(
+		{
+			&SelectionFilterSet1,
+			&SelectionFilterSet2,
+			&SelectionFilterSet3,
+			&SelectionFilterSet4,
+			&SelectionFilterSet5,
+			&SelectionFilterSet6,
+			&SelectionFilterSet7,
+			&SelectionFilterSet8,
+			&SelectionFilterSet9
+		});
+}
+
 #undef LOCTEXT_NAMESPACE
diff --git a/Engine/Plugins/Experimental/ChaosClothAssetEditor/Source/ChaosClothAssetDataflowNodes/Public/ChaosClothAsset/ConnectableValue.h b/Engine/Plugins/Experimental/ChaosClothAssetEditor/Source/ChaosClothAssetDataflowNodes/Public/ChaosClothAsset/ConnectableValue.h
--- a/Engine/Plugins/Experimental/ChaosClothAssetEditor/Source/ChaosClothAssetDataflowNodes/Public/ChaosClothAsset/ConnectableValue.h
+++ b/Engine/Plugins/Experimental/ChaosClothAssetEditor/Source/ChaosClothAssetDataflowNodes/Public/ChaosClothAsset/ConnectableValue.h
@@ -5,6 +5,20 @@
 #include "Containers/UnrealString.h"
 #include "ConnectableValue.generated.h"
 
+USTRUCT()
+struct FChaosClothAssetConnectableStringValue
+{
+	GENERATED_BODY()
+
+	/** The value for this property. */
+	UPROPERTY(EditAnywhere, Category = "Value")
+	FString StringValue = TEXT("StringValue");
+
+	/** The string value override value for when the StringValue has a connection that replaces the provided string value. */
+	UPROPERTY(VisibleAnywhere, Category = "Value", Transient)
+	mutable FString StringValue_Override;  // _Override has a special meaning to the property customization, mutable because this property is set while getting the original value
+};
+
 USTRUCT()
 struct FChaosClothAssetConnectableIStringValue
 {
diff --git a/Engine/Plugins/Experimental/ChaosClothAssetEditor/Source/ChaosClothAssetDataflowNodes/Public/ChaosClothAsset/ProxyDeformerNode.h b/Engine/Plugins/Experimental/ChaosClothAssetEditor/Source/ChaosClothAssetDataflowNodes/Public/ChaosClothAsset/ProxyDeformerNode.h
--- a/Engine/Plugins/Experimental/ChaosClothAssetEditor/Source/ChaosClothAssetDataflowNodes/Public/ChaosClothAsset/ProxyDeformerNode.h
+++ b/Engine/Plugins/Experimental/ChaosClothAssetEditor/Source/ChaosClothAssetDataflowNodes/Public/ChaosClothAsset/ProxyDeformerNode.h
@@ -3,8 +3,9 @@
 #pragma once
 
 #include "Dataflow/DataflowNode.h"
-#include "GeometryCollection/ManagedArrayCollection.h"
+#include "ChaosClothAsset/ConnectableValue.h"
 #include "ChaosClothAsset/WeightedValue.h"
+#include "GeometryCollection/ManagedArrayCollection.h"
 #include "ProxyDeformerNode.generated.h"
 
 /** Add the proxy deformer information to this cloth collection's render data. */
@@ -19,6 +20,55 @@ public:
 	UPROPERTY(Meta = (Dataflowinput, DataflowOutput, DataflowPassthrough = "Collection"))
 	FManagedArrayCollection Collection;
 
+	/**
+	 * The name of a selection containing all the dynamic points. Must be of group type SimVertices2D, SimVertices3D, or SimFaces.
+	 * Using an empty (or invalid) selection will make the proxy deformer consider all simulation points as dynamic points,
+	 * and will fully contribute to the render mesh animations (as opposed to using the render mesh skinning for the non dynamic points).
+	 * This selection is usually built from the same weight map set to the MaxDistance config using a WeightMapToSelection node and a very low threshold.
+	 */
+	UPROPERTY(EditAnywhere, Category = "ProxyDeformer")
+	FChaosClothAssetConnectableIStringValue SimVertexSelection;
+
+	/** Selection filter set used to restrict a primary selection of render vertices to a secondary selection of simulation mesh triangles. Right click and do AddPin to add more selection sets .*/
+	UPROPERTY(EditAnywhere, Category = "ProxyDeformer")
+	FChaosClothAssetConnectableIStringValue SelectionFilterSet0 = { TEXT("SelectionFilterSet") };  // Must be an IStringValue for the first element of the array
+
+	/** Selection filter set used to restrict a primary selection of render vertices to a secondary selection of simulation mesh triangles. Right click and do AddPin to add more selection sets .*/
+	UPROPERTY(EditAnywhere, Category = "ProxyDeformer")
+	FChaosClothAssetConnectableStringValue SelectionFilterSet1 = { TEXT("SelectionFilterSet") };
+
+	/** Selection filter set used to restrict a primary selection of render vertices to a secondary selection of simulation mesh triangles. Right click and do AddPin to add more selection sets .*/
+	UPROPERTY(EditAnywhere, Category = "ProxyDeformer")
+	FChaosClothAssetConnectableStringValue SelectionFilterSet2 = { TEXT("SelectionFilterSet") };
+
+	/** Selection filter set used to restrict a primary selection of render vertices to a secondary selection of simulation mesh triangles. Right click and do AddPin to add more selection sets .*/
+	UPROPERTY(EditAnywhere, Category = "ProxyDeformer")
+	FChaosClothAssetConnectableStringValue SelectionFilterSet3 = { TEXT("SelectionFilterSet") };
+
+	/** Selection filter set used to restrict a primary selection of render vertices to a secondary selection of simulation mesh triangles. Right click and do AddPin to add more selection sets .*/
+	UPROPERTY(EditAnywhere, Category = "ProxyDeformer")
+	FChaosClothAssetConnectableStringValue SelectionFilterSet4 = { TEXT("SelectionFilterSet") };
+
+	/** Selection filter set used to restrict a primary selection of render vertices to a secondary selection of simulation mesh triangles. Right click and do AddPin to add more selection sets .*/
+	UPROPERTY(EditAnywhere, Category = "ProxyDeformer")
+	FChaosClothAssetConnectableStringValue SelectionFilterSet5 = { TEXT("SelectionFilterSet") };
+
+	/** Selection filter set used to restrict a primary selection of render vertices to a secondary selection of simulation mesh triangles. Right click and do AddPin to add more selection sets .*/
+	UPROPERTY(EditAnywhere, Category = "ProxyDeformer")
+	FChaosClothAssetConnectableStringValue SelectionFilterSet6 = { TEXT("SelectionFilterSet") };
+
+	/** Selection filter set used to restrict a primary selection of render vertices to a secondary selection of simulation mesh triangles. Right click and do AddPin to add more selection sets .*/
+	UPROPERTY(EditAnywhere, Category = "ProxyDeformer")
+	FChaosClothAssetConnectableStringValue SelectionFilterSet7 = { TEXT("SelectionFilterSet") };
+
+	/** Selection filter set used to restrict a primary selection of render vertices to a secondary selection of simulation mesh triangles. Right click and do AddPin to add more selection sets .*/
+	UPROPERTY(EditAnywhere, Category = "ProxyDeformer")
+	FChaosClothAssetConnectableStringValue SelectionFilterSet8 = { TEXT("SelectionFilterSet") };
+
+	/** Selection filter set used to restrict a primary selection of render vertices to a secondary selection of simulation mesh triangles. Right click and do AddPin to add more selection sets .*/
+	UPROPERTY(EditAnywhere, Category = "ProxyDeformer")
+	FChaosClothAssetConnectableStringValue SelectionFilterSet9 = { TEXT("SelectionFilterSet") };
+
 	/** Whether using multiple simulation mesh triangles to influence the position of the deformed render vertex. */
 	UPROPERTY(EditAnywhere, Category = "ProxyDeformer")
 	bool bUseMultipleInfluences = false;
@@ -35,10 +85,6 @@ public:
 	UPROPERTY(EditAnywhere, Category = "ProxyDeformer")
 	bool bUseSmoothTransition = true;
 
-	/** The name of the simulation mesh MaxDistance weight map used to determine the kinematic/dynamic points. */
-	UPROPERTY(EditAnywhere, Category = "ProxyDeformer")
-	FChaosClothAssetWeightedValueNonAnimatableNoLowHighRange MaxDistance;
-
 	/**
 	 * The name of the render mesh weight map generated by this node detailing the contribution of the proxy deformer.
 	 * Value ranges between 0 (fully deformed) and 1 (fully skinned).
@@ -50,5 +96,23 @@ public:
 	FChaosClothAssetProxyDeformerNode(const Dataflow::FNodeParameters& InParam, FGuid InGuid = FGuid::NewGuid());
 
 private:
+	/** Hardcoded number of FilterSets since it is currently not possible to use arrays for optional inputs. */
+	static constexpr int32 MaxNumFilterSets = 10;
+
+	/** The number of filter sets currently exposed to the node UI. */
+	UPROPERTY()
+	int32 NumFilterSets = 1;
+
+	//~ Begin FDataflowNode interface
 	virtual void Evaluate(Dataflow::FContext& Context, const FDataflowOutput* Out) const override;
+	virtual Dataflow::FPin AddPin() override;
+	virtual bool CanAddPin() const override { return NumFilterSets < MaxNumFilterSets; }
+	virtual bool CanRemovePin() const override { return NumFilterSets > 1; }
+	virtual Dataflow::FPin GetPinToRemove() const override;
+	virtual void OnPinRemoved(const Dataflow::FPin& Pin) override;
+	virtual void Serialize(FArchive& Ar) override;
+	//~ End FDataflowNode interface
+
+	TArray<FName> GetSelectionFilterNames(Dataflow::FContext& Context) const;
+	TArray<const FChaosClothAssetConnectableStringValue*> Get1To9SelectionFilterSets() const;
 };
diff --git a/Engine/Source/Runtime/ClothingSystemRuntimeCommon/Private/Utils/ClothingMeshUtils.cpp b/Engine/Source/Runtime/ClothingSystemRuntimeCommon/Private/Utils/ClothingMeshUtils.cpp
--- a/Engine/Source/Runtime/ClothingSystemRuntimeCommon/Private/Utils/ClothingMeshUtils.cpp
+++ b/Engine/Source/Runtime/ClothingSystemRuntimeCommon/Private/Utils/ClothingMeshUtils.cpp
@@ -55,6 +55,27 @@ namespace ClothingMeshUtils
 		int32 GetPayload(int32 Idx) const { return Idx; }
 	};
 
+	void ClothMeshDesc::GatherAllSourceTrianglesForTargetVertex(const TArray<FMeshToMeshFilterSet>& FilterSets, int32 TargetVertex) const
+	{
+		TSet<int32> FilteredTriangleSet;
+		FilteredTriangleSet.Reserve(FilteredTriangles.Num());  // Reuse the previous number of filtered triangles as a guesstimate of the required filter size
+
+		for (const FMeshToMeshFilterSet& FilterSet : FilterSets)
+		{
+			if (FilterSet.TargetVertices.Contains(TargetVertex))
+			{
+				FilteredTriangleSet.Append(FilterSet.SourceTriangles);
+			}
+		}
+
+		if (FilteredTriangles.Num() != 0 || FilteredTriangleSet.Num() != 0)
+		{
+			bHasValidBVH = false;  // The set has likely changed, invalidate the BVH
+		}
+
+		FilteredTriangles = FilteredTriangleSet.Array();
+	}
+
 	float ClothMeshDesc::DistanceToTriangle(const FVector& Position, int32 TriangleBaseIndex) const
 	{
 		const uint32 IA = Indices[TriangleBaseIndex + 0];
@@ -73,7 +94,11 @@ namespace ClothingMeshUtils
 	{
 		ensure(HasValidMesh());
 		static const int32 MinNumTrianglesForBVHCreation = 100;
-		const int32 NumTris = Indices.Num() / 3;
+
+		const int32 NumFilteredTriangles = FilteredTriangles.Num();
+
+		const int32 NumTris = NumFilteredTriangles ? NumFilteredTriangles : Indices.Num() / 3;
+
 		if (NumTris > MinNumTrianglesForBVHCreation)
 		{
 			// This is not thread safe
@@ -82,9 +107,19 @@ namespace ClothingMeshUtils
 				TArray<FClothBvEntry> BVEntries;
 				BVEntries.Reset(NumTris);
 
-				for (int32 Tri = 0; Tri < NumTris; ++Tri)
+				if (NumFilteredTriangles)
 				{
-					BVEntries.Add({ this, Tri });
+					for (int32 Tri : FilteredTriangles)
+					{
+						BVEntries.Add({ this, Tri });
+					}
+				}
+				else
+				{
+					for (int32 Tri = 0; Tri < NumTris; ++Tri)
+					{
+						BVEntries.Add({ this, Tri });
+					}
 				}
 				BVH.Reinitialize(BVEntries);
 				bHasValidBVH = true;
@@ -106,7 +141,7 @@ namespace ClothingMeshUtils
 				return BVH.FindAllIntersections(TmpAABB);
 			}
 		}
-		return TArray<int32>();
+		return NumFilteredTriangles ? FilteredTriangles : TArray<int32>();
 	}
 
 	TConstArrayView<float> ClothMeshDesc::GetMaxEdgeLengths() const
@@ -173,15 +208,6 @@ namespace ClothingMeshUtils
 		}
 	}
 
-	// Explicit template instantiations of SkinPhysicsMesh
-	// TODO: Deprecated, remove this templated version post 5.0
-	template
-	void CLOTHINGSYSTEMRUNTIMECOMMON_API SkinPhysicsMesh<true, false>(const TArray<int32>& BoneMap, const FClothPhysicalMeshData& InMesh, const FTransform& RootBoneTransform,
-		const FMatrix44f* InBoneMatrices, const int32 InNumBoneMatrices, TArray<FVector3f>& OutPositions, TArray<FVector3f>& OutNormals, uint32 ArrayOffset);
-	template
-	void CLOTHINGSYSTEMRUNTIMECOMMON_API SkinPhysicsMesh<false, true>(const TArray<int32>& BoneMap, const FClothPhysicalMeshData& InMesh, const FTransform& RootBoneTransform,
-		const FMatrix44f* InBoneMatrices, const int32 InNumBoneMatrices, TArray<FVector3f>& OutPositions, TArray<FVector3f>& OutNormals, uint32 ArrayOffset);
-
 	// inline function used to force the unrolling of the skinning loop
 	FORCEINLINE static void AddInfluence(FVector3f& OutPosition, FVector3f& OutNormal, const FVector3f& RefParticle, const FVector3f& RefNormal, const FMatrix44f& BoneMatrix, const float Weight)
 	{
@@ -259,99 +285,6 @@ namespace ClothingMeshUtils
 		}, NumVerts > MinParallelVertices ? EParallelForFlags::None : EParallelForFlags::ForceSingleThread);
 	}
 
-	// TODO: Deprecated, remove this templated version post 5.0
-	template<bool bInPlaceOutput, bool bRemoveScaleAndInvertPostTransform>
-	void SkinPhysicsMesh(
-		const TArray<int32>& InBoneMap,
-		const FClothPhysicalMeshData& InMesh,
-		const FTransform& PostTransform,
-		const FMatrix44f* InBoneMatrices,
-		const int32 InNumBoneMatrices,
-		TArray<FVector3f>& OutPositions,
-		TArray<FVector3f>& OutNormals,
-		uint32 ArrayOffset)
-	{
-		const uint32 NumVerts = InMesh.Vertices.Num();
-
-		if(!bInPlaceOutput)
-		{
-			ensure(ArrayOffset == 0);
-			OutPositions.Reset(NumVerts);
-			OutNormals.Reset(NumVerts);
-			OutPositions.AddZeroed(NumVerts);
-			OutNormals.AddZeroed(NumVerts);
-		}
-		else
-		{
-			check((uint32) OutPositions.Num() >= NumVerts + ArrayOffset);
-			check((uint32) OutNormals.Num() >= NumVerts + ArrayOffset);
-			// It is faster on some platforms to zero the memory first instead of changing this function to work with uninitialized memory
-			FMemory::Memzero((uint8*)OutPositions.GetData() + ArrayOffset * sizeof(FVector3f), NumVerts * sizeof(FVector3f));
-			FMemory::Memzero((uint8*)OutNormals.GetData() + ArrayOffset * sizeof(FVector3f), NumVerts * sizeof(FVector3f));
-		}
-
-		const int32 MaxInfluences = InMesh.MaxBoneWeights;
-		UE_CLOG(MaxInfluences > 12, LogClothingMeshUtils, Warning, TEXT("The cloth physics mesh skinning code can't cope with more than 12 bone influences."));
-
-		const int32* const RESTRICT BoneMap = InBoneMap.GetData();  // Remove RangeCheck for faster skinning in development builds
-		const FMatrix44f* const RESTRICT BoneMatrices = InBoneMatrices;
-		
-		static const uint32 MinParallelVertices = 500;  // 500 seems to be the lowest threshold still giving gains even on profiled assets that are only using a small number of influences
-
-		ParallelFor(NumVerts, [&InMesh, &PostTransform, BoneMap, BoneMatrices, &OutPositions, &OutNormals, ArrayOffset](uint32 VertIndex)
-		{
-			// Fixed particle, needs to be skinned
-			const uint16* const RESTRICT BoneIndices = InMesh.BoneData[VertIndex].BoneIndices;
-			const float* const RESTRICT BoneWeights = InMesh.BoneData[VertIndex].BoneWeights;
-
-			// WARNING - HORRIBLE UNROLLED LOOP + JUMP TABLE BELOW
-			// done this way because this is a pretty tight and perf critical loop. essentially
-			// rather than checking each influence we can just jump into this switch and fall through
-			// everything to compose the final skinned data
-			const FVector3f& RefParticle = InMesh.Vertices[VertIndex];
-			const FVector3f& RefNormal = InMesh.Normals[VertIndex];
-			FVector3f& OutPosition = OutPositions[bInPlaceOutput ? VertIndex + ArrayOffset : VertIndex];
-			FVector3f& OutNormal = OutNormals[bInPlaceOutput ? VertIndex + ArrayOffset : VertIndex];
-			switch (InMesh.BoneData[VertIndex].NumInfluences)
-			{
-			case 12: AddInfluence(OutPosition, OutNormal, RefParticle, RefNormal, BoneMatrices[BoneMap[BoneIndices[11]]], BoneWeights[11]);  // Intentional fall through
-			case 11: AddInfluence(OutPosition, OutNormal, RefParticle, RefNormal, BoneMatrices[BoneMap[BoneIndices[10]]], BoneWeights[10]);  // Intentional fall through
-			case 10: AddInfluence(OutPosition, OutNormal, RefParticle, RefNormal, BoneMatrices[BoneMap[BoneIndices[ 9]]], BoneWeights[ 9]);  // Intentional fall through
-			case  9: AddInfluence(OutPosition, OutNormal, RefParticle, RefNormal, BoneMatrices[BoneMap[BoneIndices[ 8]]], BoneWeights[ 8]);  // Intentional fall through
-			case  8: AddInfluence(OutPosition, OutNormal, RefParticle, RefNormal, BoneMatrices[BoneMap[BoneIndices[ 7]]], BoneWeights[ 7]);  // Intentional fall through
-			case  7: AddInfluence(OutPosition, OutNormal, RefParticle, RefNormal, BoneMatrices[BoneMap[BoneIndices[ 6]]], BoneWeights[ 6]);  // Intentional fall through
-			case  6: AddInfluence(OutPosition, OutNormal, RefParticle, RefNormal, BoneMatrices[BoneMap[BoneIndices[ 5]]], BoneWeights[ 5]);  // Intentional fall through
-			case  5: AddInfluence(OutPosition, OutNormal, RefParticle, RefNormal, BoneMatrices[BoneMap[BoneIndices[ 4]]], BoneWeights[ 4]);  // Intentional fall through
-			case  4: AddInfluence(OutPosition, OutNormal, RefParticle, RefNormal, BoneMatrices[BoneMap[BoneIndices[ 3]]], BoneWeights[ 3]);  // Intentional fall through
-			case  3: AddInfluence(OutPosition, OutNormal, RefParticle, RefNormal, BoneMatrices[BoneMap[BoneIndices[ 2]]], BoneWeights[ 2]);  // Intentional fall through
-			case  2: AddInfluence(OutPosition, OutNormal, RefParticle, RefNormal, BoneMatrices[BoneMap[BoneIndices[ 1]]], BoneWeights[ 1]);  // Intentional fall through
-			case  1: AddInfluence(OutPosition, OutNormal, RefParticle, RefNormal, BoneMatrices[BoneMap[BoneIndices[ 0]]], BoneWeights[ 0]);  // Intentional fall through
-			default: break;
-			}
-
-			if (bRemoveScaleAndInvertPostTransform)
-			{
-				// Ignore any user scale. It's already accounted for in our skinning matrices
-				// This is the use case for NVcloth
-				FTransform PostTransformInternal = PostTransform;
-				PostTransformInternal.SetScale3D(FVector(1.0f));
-
-				OutPosition = (FVector3f)PostTransformInternal.InverseTransformPosition((FVector)OutPosition);
-				OutNormal = (FVector3f)PostTransformInternal.InverseTransformVector((FVector)OutNormal);
-			}
-			else
-			{
-				OutPosition = (FVector3f)PostTransform.TransformPosition((FVector)OutPosition);
-				OutNormal = (FVector3f)PostTransform.TransformVector((FVector)OutNormal);
-			}
-
-			if (OutNormal.SizeSquared() > SMALL_NUMBER)
-			{
-				OutNormal = OutNormal.GetUnsafeNormal();
-			}
-		}, NumVerts > MinParallelVertices ? EParallelForFlags::None : EParallelForFlags::ForceSingleThread);
-	}
-
 	/** 
 	 * Gets the best match triangle for a specified position from the triangles in Mesh.
 	 * Performs no validation on the incoming mesh data, the mesh data should be verified
@@ -404,11 +337,21 @@ namespace ClothingMeshUtils
 			const TArray<int32> Tris = Mesh.FindCandidateTriangles(Position);
 			int32 NumTriangles = Tris.Num();
 
+			const TArray<int32>& FilteredTriangles = Mesh.GetFilteredTriangles();
+			const int32 NumFilteredTriangles = FilteredTriangles.Num();
+			const bool bHasFilteredTriangles = (NumFilteredTriangles > 0);
+
+			auto GetTriangleIndex = [bHasFilteredTriangles, &FilteredTriangles](int32 Index) -> int32
+				{
+					return bHasFilteredTriangles ? FilteredTriangles[Index] : Index;
+				};
+
 			if (NumTriangles < N)
 			{
 				// Couldn't find N candidates using FindCandidateTriangles. Grab all triangles in the mesh and see if
 				// we have enough
-				NumTriangles = Mesh.GetIndices().Num() / 3;
+
+				NumTriangles = bHasFilteredTriangles ? NumFilteredTriangles : Mesh.GetIndices().Num() / 3;
 
 				if (NumTriangles < N)
 				{
@@ -418,8 +361,8 @@ namespace ClothingMeshUtils
 					int32 i = 0;
 					for ( ; i < NumTriangles; ++i)
 					{
-						int32 TriBaseIdx = 3*i;
-						float CurrentDistance = Mesh.DistanceToTriangle(Position, TriBaseIdx);
+						const int32 TriBaseIdx = 3 * GetTriangleIndex(i);
+						const float CurrentDistance = Mesh.DistanceToTriangle(Position, TriBaseIdx);
 						ClosestTriangles[i] = TPair<int32, float>{ TriBaseIdx, CurrentDistance };
 					}
 
@@ -442,7 +385,7 @@ namespace ClothingMeshUtils
 			int32 i = 0;
 			for (; i < N; ++i)
 			{
-				int32 TriBaseIdx = (Tris.Num() >= N ? Tris[i] : i) * 3;
+				int32 TriBaseIdx = (Tris.Num() >= N ? Tris[i] : GetTriangleIndex(i)) * 3;
 				float CurrentDistance = Mesh.DistanceToTriangle(Position, TriBaseIdx);
 				ClosestTriangles[i] = TPair<int32, float>{ TriBaseIdx, CurrentDistance };
 			}
@@ -456,7 +399,7 @@ namespace ClothingMeshUtils
 			// Now keep going
 			for (; i < NumTriangles; ++i)
 			{
-				int32 TriBaseIdx = (Tris.Num() >= N ? Tris[i] : i) * 3;
+				int32 TriBaseIdx = (Tris.Num() >= N ? Tris[i] : GetTriangleIndex(i)) * 3;
 				float CurrentDistance = Mesh.DistanceToTriangle(Position, TriBaseIdx);
 
 				if (CurrentDistance < ClosestTriangles[0].Value)
@@ -635,7 +578,11 @@ namespace ClothingMeshUtils
 			return true;
 		}
 
-		void FixZeroWeightVertices(TArray<FMeshToMeshVertData>& InOutSkinningData, const ClothMeshDesc& TargetMesh, const ClothMeshDesc& SourceMesh)
+		void FixZeroWeightVertices(
+			TArray<FMeshToMeshVertData>& InOutSkinningData,
+			const ClothMeshDesc& TargetMesh,
+			const ClothMeshDesc& SourceMesh,
+			const TArray<FMeshToMeshFilterSet>& FilterSets)
 		{
 			if (!ensure(InOutSkinningData.Num() > 0))
 			{
@@ -735,7 +682,8 @@ namespace ClothingMeshUtils
 		const FPointWeightMap* MaxDistances,
 		bool bUseSmoothTransitions,
 		bool bUseMultipleInfluences,
-		float KernelMaxDistance)
+		float KernelMaxDistance,
+		const TArray<FMeshToMeshFilterSet>& FilterSets)
 	{
 		check(SourceMesh.HasAveragedNormals());
 
@@ -752,13 +700,6 @@ namespace ClothingMeshUtils
 		}
 
 		const int32 NumMesh0Verts = TargetMesh.GetPositions().Num();
-		const int32 NumMesh0Normals = TargetMesh.GetNormals().Num();
-		const int32 NumMesh0Tangents = TargetMesh.HasTangents() ? TargetMesh.GetTangents().Num() : 0;
-		const int32 NumMesh0Tris = TargetMesh.GetIndices().Num() / 3;
-
-		const int32 NumMesh1Verts = SourceMesh.GetPositions().Num();
-		const int32 NumMesh1Normals = SourceMesh.GetNormals().Num();
-		const int32 NumMesh1Indices = SourceMesh.GetIndices().Num();
 
 		const TConstArrayView<float> MaxEdgeLengths = TargetMesh.GetMaxEdgeLengths();
 		check(NumMesh0Verts == MaxEdgeLengths.Num());
@@ -775,6 +716,8 @@ namespace ClothingMeshUtils
 			// For all mesh0 verts
 			for (int32 VertIdx0 = 0; VertIdx0 < NumMesh0Verts; ++VertIdx0)
 			{
+				SourceMesh.GatherAllSourceTrianglesForTargetVertex(FilterSets, VertIdx0);
+
 				TStaticArray<FMeshToMeshVertData, NUM_INFLUENCES_PER_VERTEX> SkinningData;
 				const bool bOK = MultipleSkinningDataForVertex(
 					SkinningData,
@@ -810,6 +753,8 @@ namespace ClothingMeshUtils
 			bool bAnySmallTriangleEncountered = false;
 			for (int32 VertIdx0 = 0; VertIdx0 < NumMesh0Verts; ++VertIdx0)
 			{
+				SourceMesh.GatherAllSourceTrianglesForTargetVertex(FilterSets, VertIdx0);
+
 				OutMeshToMeshVertData.AddZeroed();
 				FMeshToMeshVertData& SkinningData = OutMeshToMeshVertData.Last();
 
@@ -862,7 +807,7 @@ namespace ClothingMeshUtils
 		{
 			if (bUseMultipleInfluences)
 			{
-				FixZeroWeightVertices(OutMeshToMeshVertData, TargetMesh, SourceMesh);
+				FixZeroWeightVertices(OutMeshToMeshVertData, TargetMesh, SourceMesh, FilterSets);
 			}
 
 			if (MaxDistances)
@@ -872,104 +817,6 @@ namespace ClothingMeshUtils
 		}
 	}
 
-	// TODO: Deprecated, old public interface, remove after 5.0
-	void FixZeroWeightVertices(TArray<FMeshToMeshVertData>& InOutSkinningData,
-							   const ClothMeshDesc& TargetMesh,
-							   const TArray<FVector3f>* TargetTangents,
-							   const ClothMeshDesc& SourceMesh,
-							   const TArray<float>& MaxEdgeLength)
-	{
-		if (!ensure(InOutSkinningData.Num() > 0))
-		{
-			return;
-		}
-
-		const int32 NumTargetMeshVerts = TargetMesh.GetPositions().Num();
-
-		if (!ensure(NumTargetMeshVerts * NUM_INFLUENCES_PER_VERTEX == InOutSkinningData.Num()))
-		{
-			return;
-		}
-
-		if (!ensure(NumTargetMeshVerts == MaxEdgeLength.Num()))
-		{
-			return;
-		}
-
-		bool bAnySmallTriangleEncountered = false;
-
-		for (int32 VID = 0; VID < NumTargetMeshVerts; ++VID)
-		{
-			float SumWeight = 0.0f;
-			for (int Inf = 0; Inf < NUM_INFLUENCES_PER_VERTEX; ++Inf)
-			{
-				const FMeshToMeshVertData& Data = InOutSkinningData[NUM_INFLUENCES_PER_VERTEX * VID + Inf];
-				if (Data.SourceMeshVertIndices[3] < 0xFFFF)
-				{
-					SumWeight += Data.Weight;
-				}
-			}
-
-			if (SumWeight < KINDA_SMALL_NUMBER)
-			{
-				// Fall back to single influence
-				const FVector3f& VertPosition = TargetMesh.GetPositions()[VID];
-				const FVector3f& VertNormal = TargetMesh.GetNormals()[VID];
-
-				FVector3f VertTangent;
-				if (TargetTangents)
-				{
-					VertTangent = (*TargetTangents)[VID];
-				}
-				else
-				{
-					FVector3f Tan0, Tan1;
-					VertNormal.FindBestAxisVectors(Tan0, Tan1);
-					VertTangent = Tan0;
-				}
-
-				// Compute single-influence attachment for the first skinning data
-				FMeshToMeshVertData& FirstData = InOutSkinningData[NUM_INFLUENCES_PER_VERTEX * VID];
-
-				const int32 ClosestTriangleBaseIdx = ClothingMeshUtils::GetBestTriangleBaseIndex(SourceMesh, (FVector)VertPosition, MaxEdgeLength[VID]);
-				if (!ensure(ClosestTriangleBaseIdx != INDEX_NONE))
-				{
-					FirstData.Weight = 0.0f;
-				}
-				else
-				{
-					const uint16 PreviousFlag = FirstData.SourceMeshVertIndices[3];
-					const bool bSkinningSuccess = SingleSkinningDataForVertex(VertPosition, VertNormal, VertTangent, SourceMesh, ClosestTriangleBaseIdx, FirstData);
-					FirstData.SourceMeshVertIndices[3] = PreviousFlag;
-
-					if (!bSkinningSuccess)
-					{
-						bAnySmallTriangleEncountered = true;
-					}
-				}
-
-				// Set all other skinning data to have zero weight
-				for (int Inf = 1; Inf < NUM_INFLUENCES_PER_VERTEX; ++Inf)
-				{
-					FMeshToMeshVertData& CurrentData = InOutSkinningData[NUM_INFLUENCES_PER_VERTEX * VID + Inf];
-					CurrentData.SourceMeshVertIndices[3] = 0xFFFF;
-					CurrentData.Weight = 0.0f;
-				}
-			}
-		}
-
-#if WITH_EDITOR
-		if (bAnySmallTriangleEncountered)
-		{
-			const FText ErrorMsg = LOCTEXT("DegenerateTriangleErrorToast", "Failed to generate skinning data, found conincident vertices in at least one triangle. See Log for details");
-			FNotificationInfo Info(ErrorMsg);
-			Info.ExpireDuration = 5.0f;
-			FSlateNotificationManager::Get().AddNotification(Info);
-		}
-#endif
-
-	}
-
 	// TODO: Vertex normals are not used at present, a future improved algorithm might however
 	FVector4f GetPointBaryAndDist(const FVector3f& A, const FVector3f& B, const FVector3f& C, const FVector3f& Point)
 	{
@@ -1292,48 +1139,6 @@ namespace ClothingMeshUtils
 		return BaryAndDist;
 	}
 
-
-	// Deprecated, remove after 5.0
-	void GenerateEmbeddedPositions(const ClothMeshDesc& SourceMesh, TArrayView<const FVector3f> Positions, TArray<FVector4>& OutEmbeddedPositions, TArray<int32>& OutSourceIndices)
-	{
-		if(!SourceMesh.HasValidMesh())
-		{
-			// No valid source mesh
-			return;
-		}
-
-		const int32 NumPositions = Positions.Num();
-
-		OutEmbeddedPositions.Reset();
-		OutEmbeddedPositions.AddUninitialized(NumPositions);
-
-		OutSourceIndices.Reset(NumPositions * 3);
-
-		for(int32 PositionIndex = 0 ; PositionIndex < NumPositions ; ++PositionIndex)
-		{
-			const FVector3f& Position = Positions[PositionIndex];
-
-			int32 TriBaseIndex = GetBestTriangleBaseIndex(SourceMesh, (FVector)Position);
-
-			const int32 IA = SourceMesh.GetIndices()[TriBaseIndex];
-			const int32 IB = SourceMesh.GetIndices()[TriBaseIndex + 1];
-			const int32 IC = SourceMesh.GetIndices()[TriBaseIndex + 2];
-
-			const FVector3f& A = SourceMesh.GetPositions()[IA];
-			const FVector3f& B = SourceMesh.GetPositions()[IB];
-			const FVector3f& C = SourceMesh.GetPositions()[IC];
-
-			const FVector3f& NA = SourceMesh.GetNormals()[IA];
-			const FVector3f& NB = SourceMesh.GetNormals()[IB];
-			const FVector3f& NC = SourceMesh.GetNormals()[IC];
-
-			OutEmbeddedPositions[PositionIndex] = (FVector4)GetPointBaryAndDistWithNormals(A, B, C, NA, NB, NC, Position);
-			OutSourceIndices.Add(IA);
-			OutSourceIndices.Add(IB);
-			OutSourceIndices.Add(IC);
-		}
-	}
-
 	void ComputeVertexContributions(
 		TArray<FMeshToMeshVertData>& InOutSkinningData,
 		const FPointWeightMap* const InMaxDistances,
diff --git a/Engine/Source/Runtime/ClothingSystemRuntimeCommon/Public/Utils/ClothingMeshUtils.h b/Engine/Source/Runtime/ClothingSystemRuntimeCommon/Public/Utils/ClothingMeshUtils.h
--- a/Engine/Source/Runtime/ClothingSystemRuntimeCommon/Public/Utils/ClothingMeshUtils.h
+++ b/Engine/Source/Runtime/ClothingSystemRuntimeCommon/Public/Utils/ClothingMeshUtils.h
@@ -12,6 +12,12 @@ struct FPointWeightMap;
 
 namespace ClothingMeshUtils
 {
+	struct FMeshToMeshFilterSet
+	{
+		TSet<int32> SourceTriangles;  // Set of triangle index in the source ClothMeshDesc
+		TSet<int32> TargetVertices;  // Set of vertex index in the target ClothMeshDesc
+	};
+
 	class ClothMeshDesc
 	{
 	public:
@@ -59,6 +65,12 @@ namespace ClothingMeshUtils
 		/** Find the distance to the specified triangle from a point. */
 		CLOTHINGSYSTEMRUNTIMECOMMON_API float DistanceToTriangle(const FVector& Position, int32 TriangleBaseIndex) const;
 
+		/** Return the current list of filtered triangles. */
+		TArray<int32>& GetFilteredTriangles() const { return FilteredTriangles; }
+
+		/** Create an array with all source triangles whose filter set from the given array containts the specified target vertex. */
+		void GatherAllSourceTrianglesForTargetVertex(const TArray<FMeshToMeshFilterSet>& FilterSets, int32 TargetVertex) const;
+
 		/** Find the closest triangles from the specified point. */
 		CLOTHINGSYSTEMRUNTIMECOMMON_API TArray<int32> FindCandidateTriangles(const FVector& InPoint, float InTolerance = KINDA_SMALL_NUMBER) const;
 
@@ -84,6 +96,7 @@ namespace ClothingMeshUtils
 		mutable Chaos::TAABBTree<int32, Chaos::TAABBTreeLeafArray<int32, false>, false> BVH;
 
 		mutable TArray<float> MaxEdgeLengths;
+		mutable TArray<int32> FilteredTriangles;
 	};
 
 	/**
@@ -98,24 +111,6 @@ namespace ClothingMeshUtils
 		TArray<FVector3f>& OutPositions, 
 		TArray<FVector3f>& OutNormals);
 
-	/**
-	 * Static method for calculating a skinned mesh result from source data
-	 * The bInPlaceOutput allows us to directly populate arrays that are already allocated
-	 * bRemoveScaleAndInvertPostTransform will determine if the PostTransform should be inverted and the scale removed (NvCloth uses this).
-	 * It is templated to remove branches at compile time
-	 */
-	template<bool bInPlaceOutput = false, bool bRemoveScaleAndInvertPostTransform = true>
-	UE_DEPRECATED(5.0, "Use non templated version of SkinPhysicsMesh instead.")
-	void CLOTHINGSYSTEMRUNTIMECOMMON_API SkinPhysicsMesh(
-		const TArray<int32>& BoneMap,
-		const FClothPhysicalMeshData& InMesh,
-		const FTransform& PostTransform,
-		const FMatrix44f* InBoneMatrices,
-		const int32 InNumBoneMatrices,
-		TArray<FVector3f>& OutPositions,
-		TArray<FVector3f>& OutNormals,
-		uint32 ArrayOffset);
-
 	/**
 	* Given mesh information for two meshes, generate a list of skinning data to embed TargetMesh in SourceMesh
 	* 
@@ -126,6 +121,7 @@ namespace ClothingMeshUtils
 	* @param bUseSmoothTransitions      - Set blend weight to smoothen transitions between the fixed and deformed vertices when updating the vertex contributions
 	* @param bUseMultipleInfluences     - Whether to take a weighted average of influences from multiple source triangles
 	* @param KernelMaxDistance          - Max distance parameter for weighting kernel for when using multiple influences
+	* @param FilterSets                 - A set of source triangle indices to filter the search from, will use all source triangles if empty
 	*/
 	void CLOTHINGSYSTEMRUNTIMECOMMON_API GenerateMeshToMeshVertData(
 		TArray<FMeshToMeshVertData>& OutMeshToMeshVertData,
@@ -134,58 +130,8 @@ namespace ClothingMeshUtils
 		const FPointWeightMap* MaxDistances,
 		bool bUseSmoothTransitions,
 		bool bUseMultipleInfluences,
-		float KernelMaxDistance);
-
-	/**
-	 * Compute the max edge length for each edge coming out of a mesh vertex. Useful for guiding the search radius when
-	 * searching for nearest triangles.
-	 */
-	UE_DEPRECATED(5.0, "Superseded by GenerateMeshToMeshVertData.")
-	inline void ComputeMaxEdgeLength(const ClothMeshDesc& TargetMesh, TArray<float>& OutMaxEdgeLength)
-	{
-		OutMaxEdgeLength = TargetMesh.GetMaxEdgeLengths();
-	}
-
-	/**
-	* Given mesh information for two meshes, generate a list of skinning data to embed TargetMesh in SourceMesh
-	* 
-	* @param OutSkinningData            - Final skinning data
-	* @param TargetMesh                 - Mesh data for the mesh we are embedding
-	* @param TargetTangents             - Optional Tangents for the mesh we are embedding
-	* @param SourceMesh                 - Mesh data for the mesh we are embedding into
-	* @param TargetMaxEdgeLength        - Per-vertex longest incident edge length, as returned by ComputeMaxEdgeLength()
-	* @param bUseMultipleInfluences     - Whether to take a weighted average of influences from multiple source triangles
-	* @param KernelMaxDistance          - Max distance parameter for weighting kernel
-	*/
-	UE_DEPRECATED(5.0, "Superseded by GenerateMeshToMeshVertData.")
-	inline void GenerateMeshToMeshSkinningData(
-		TArray<FMeshToMeshVertData>& OutSkinningData,
-		const ClothMeshDesc& TargetMesh,
-		const TArray<FVector3f>* TargetTangents,
-		const ClothMeshDesc& SourceMesh,
-		const TArray<float>& TargetMaxEdgeLength,
-		bool bUseMultipleInfluences,
-		float KernelMaxDistance)
-	{
-		const FPointWeightMap* const MaxDistances = nullptr;  // Disables compute vertex contributions
-		constexpr bool bUseSmoothTransitions = false;
-		
-		GenerateMeshToMeshVertData(OutSkinningData, TargetMesh, SourceMesh, MaxDistances, bUseSmoothTransitions, bUseMultipleInfluences, KernelMaxDistance);
-	}
-
-	/** 
-	 * Embeds a list of positions into a source mesh
-	 * @param SourceMesh The mesh to embed in
-	 * @param Positions The positions to embed in SourceMesh
-	 * @param OutEmbeddedPositions Embedded version of the original positions, a barycentric coordinate and distance along the normal of the triangle
-	 * @param OutSourceIndices Source index list for the embedded positions, 3 per position to denote the source triangle
-	 */
-	UE_DEPRECATED(5.0, "Use FVertexParameterMapper::GenerateEmbeddedPositions instead.")
-	void CLOTHINGSYSTEMRUNTIMECOMMON_API GenerateEmbeddedPositions(
-		const ClothMeshDesc& SourceMesh, 
-		TConstArrayView<FVector3f> Positions, 
-		TArray<FVector4>& OutEmbeddedPositions, 
-		TArray<int32>& OutSourceIndices);
+		float KernelMaxDistance,
+		const TArray<FMeshToMeshFilterSet>& FilterSets = TArray<FMeshToMeshFilterSet>());
 
 	/**
 	 * Computes how much each vertex contributes to the final mesh. The final mesh is a blend
@@ -197,18 +143,6 @@ namespace ClothingMeshUtils
 		const bool bInSmoothTransition,
 		const bool bInUseMultipleInfluences = false);
 
-	/**
-	 * Identify vertices that are not influenced by any triangles, and compute a new single attachment for the
-	 * vertex.
-	 */
-	UE_DEPRECATED(5.0, "Superseded by GenerateMeshToMeshVertData.")
-	void CLOTHINGSYSTEMRUNTIMECOMMON_API FixZeroWeightVertices(
-		TArray<FMeshToMeshVertData>& InOutSkinningData,
-		const ClothMeshDesc& TargetMesh,
-		const TArray<FVector3f>* TargetTangents,
-		const ClothMeshDesc& SourceMesh,
-		const TArray<float>& TargetMaxEdgeLength);
-
 	/**
 	* Given a triangle ABC with normals at each vertex NA, NB and NC, get a barycentric coordinate
 	* and corresponding distance from the triangle encoded in an FVector4 where the components are
