Queue notification in case we are not being called on the main game thread.

diff --git a/Engine/Source/Runtime/ClothingSystemRuntimeCommon/Private/Utils/ClothingMeshUtils.cpp b/Engine/Source/Runtime/ClothingSystemRuntimeCommon/Private/Utils/ClothingMeshUtils.cpp
--- a/Engine/Source/Runtime/ClothingSystemRuntimeCommon/Private/Utils/ClothingMeshUtils.cpp
+++ b/Engine/Source/Runtime/ClothingSystemRuntimeCommon/Private/Utils/ClothingMeshUtils.cpp
@@ -795,9 +795,10 @@ namespace ClothingMeshUtils
 			if (bAnySmallTriangleEncountered)
 			{
 				const FText ErrorMsg = LOCTEXT("DegenerateTriangleErrorToast", "Failed to generate skinning data, found conincident vertices in at least one triangle. See Log for details");
-				FNotificationInfo Info(ErrorMsg);
-				Info.ExpireDuration = 5.0f;
-				FSlateNotificationManager::Get().AddNotification(Info);
+				FNotificationInfo *Info = new FNotificationInfo(ErrorMsg);
+				Info->ExpireDuration = 5.0f;
+				// Queue notification because we may not be on the main game thread
+				FSlateNotificationManager::Get().QueueNotification(Info);
 			}
 #endif
 
