SkeletalMesh.h is needed for all builds, not just Editor builds.

diff --git a/Engine/Source/Runtime/ClothingSystemRuntimeCommon/Private/ClothingAsset.cpp b/Engine/Source/Runtime/ClothingSystemRuntimeCommon/Private/ClothingAsset.cpp
--- a/Engine/Source/Runtime/ClothingSystemRuntimeCommon/Private/ClothingAsset.cpp
+++ b/Engine/Source/Runtime/ClothingSystemRuntimeCommon/Private/ClothingAsset.cpp
@@ -8,9 +8,7 @@
 #include "Utils/ClothingMeshUtils.h"
 #include "Features/IModularFeatures.h"
 
-#if WITH_EDITOR
 #include "Engine/SkeletalMesh.h"
-#endif
 
 #include "PhysicsEngine/PhysicsAsset.h"
 
