Cloth - Fixed failed shared config migration on legacy clothing data.

[RN] Cloth - Fixed failed shared config migration on some legacy clothing data when reloading old Skeletal Mesh assets.

diff --git a/Engine/Source/Runtime/ClothingSystemRuntimeCommon/Private/ClothingAsset.cpp b/Engine/Source/Runtime/ClothingSystemRuntimeCommon/Private/ClothingAsset.cpp
--- a/Engine/Source/Runtime/ClothingSystemRuntimeCommon/Private/ClothingAsset.cpp
+++ b/Engine/Source/Runtime/ClothingSystemRuntimeCommon/Private/ClothingAsset.cpp
@@ -1258,7 +1258,9 @@ bool UClothingAssetCommon::AddClothConfigs()
 			for (TSubclassOf<UClothConfigBase> ClothConfigClass : ClothingSimulationFactory->GetClothConfigClasses())
 			{
 				const FName ClothConfigName = ClothConfigClass->GetFName();
-				if (!ClothConfigs.Find(ClothConfigName))
+
+				TObjectPtr<UClothConfigBase>* const ClothConfigPtr = ClothConfigs.Find(ClothConfigName);
+				if (!ClothConfigPtr || !*ClothConfigPtr)
 				{
 					// Create new config object
 					check(!StaticFindObject(ClothConfigClass, this, *ClothConfigClass->GetName(), true));
@@ -1283,9 +1285,16 @@ bool UClothingAssetCommon::AddClothConfigs()
 						}
 					}
 
-					// Add the new config
+					// Set the new config
 					check(ClothConfig);
-					ClothConfigs.Add(ClothConfigName, ClothConfig);
+					if (ClothConfigPtr)
+					{
+						*ClothConfigPtr = ClothConfig;
+					}
+					else
+					{
+						ClothConfigs.Emplace(ClothConfigName, ClothConfig);
+					}
 					bNewConfigAdded = true;
 				}
 			}
