#UE Add more info to an assert

diff --git a/Engine/Source/Runtime/ClothingSystemRuntimeCommon/Private/ClothingAsset.cpp b/Engine/Source/Runtime/ClothingSystemRuntimeCommon/Private/ClothingAsset.cpp
--- a/Engine/Source/Runtime/ClothingSystemRuntimeCommon/Private/ClothingAsset.cpp
+++ b/Engine/Source/Runtime/ClothingSystemRuntimeCommon/Private/ClothingAsset.cpp
@@ -1263,7 +1263,7 @@ bool UClothingAssetCommon::AddClothConfigs()
 				if (!ClothConfigPtr || !*ClothConfigPtr)
 				{
 					// Create new config object
-					check(!StaticFindObject(ClothConfigClass, this, *ClothConfigClass->GetName(), true));
+					checkf(!StaticFindObject(ClothConfigClass, this, *ClothConfigClass->GetName(), true), TEXT("Unable to find ClothConfigClass in %s"), *GetPathNameSafe(this));
 					UClothConfigBase* const ClothConfig = NewObject<UClothConfigBase>(this, ClothConfigClass, ClothConfigClass->GetFName(), RF_Transactional);
 
 					// Use the legacy config struct to try find a common config as an acceptable migration source
