Chaos Cloth: Add new optional map to use to generate tethers rather than MaxDistance.

Also added option to copy a mask from another mask in the right click menu of the mask UI.

diff --git a/Engine/Plugins/ChaosCloth/Source/ChaosCloth/Private/ChaosCloth/ChaosWeightMapTarget.h b/Engine/Plugins/ChaosCloth/Source/ChaosCloth/Private/ChaosCloth/ChaosWeightMapTarget.h
--- a/Engine/Plugins/ChaosCloth/Source/ChaosCloth/Private/ChaosCloth/ChaosWeightMapTarget.h
+++ b/Engine/Plugins/ChaosCloth/Source/ChaosCloth/Private/ChaosCloth/ChaosWeightMapTarget.h
@@ -14,9 +14,11 @@ enum class EChaosWeightMapTarget : uint8
 	BackstopDistance   = (uint8)EWeightMapTargetCommon::BackstopDistance,
 	BackstopRadius     = (uint8)EWeightMapTargetCommon::BackstopRadius,
 	AnimDriveStiffness = (uint8)EWeightMapTargetCommon::AnimDriveStiffness,
+	TetherEndsMask     = (uint8)EWeightMapTargetCommon::TetherEndsMask,
+
 	// Add Chaos specific maps below this line
-	AnimDriveDamping,
-	TetherStiffness,
+	AnimDriveDamping   = (uint8)EWeightMapTargetCommon::AnimDriveDamping_DEPRECATED,
+	TetherStiffness    = (uint8)EWeightMapTargetCommon::FirstUserTarget,
 	TetherScale,
 	Drag,
 	Lift,
@@ -27,5 +29,9 @@ enum class EChaosWeightMapTarget : uint8
 	Pressure,
 	FlatnessRatio,
 	OuterDrag,
-	OuterLift
+	OuterLift,
+
+	// Add Chaos specific maps above this line
+	MAX UMETA(Hidden)
 };
+static_assert((uint8)EChaosWeightMapTarget::MAX <= (uint8)EWeightMapTargetCommon::LastUserTarget);
diff --git a/Engine/Source/Editor/ClothPainter/Private/SClothAssetSelector.cpp b/Engine/Source/Editor/ClothPainter/Private/SClothAssetSelector.cpp
--- a/Engine/Source/Editor/ClothPainter/Private/SClothAssetSelector.cpp
+++ b/Engine/Source/Editor/ClothPainter/Private/SClothAssetSelector.cpp
@@ -437,7 +437,7 @@ public:
 				const FPointWeightMap* const Mask = Item->GetMask();
 				if(Enum && Mask)
 				{
-					return SNew(STextBlock).Text(Enum->GetDisplayNameTextByIndex((int32)Mask->CurrentTarget));
+					return SNew(STextBlock).Text(Enum->GetDisplayNameTextByValue((int64)Mask->CurrentTarget));
 				}
 			}
 		}
@@ -485,6 +485,7 @@ public:
 				{
 					Builder.AddSubMenu(LOCTEXT("MaskActions_SetTarget", "Set Target"), LOCTEXT("MaskActions_SetTarget_Tooltip", "Choose the target for this mask"), FNewMenuDelegate::CreateSP(this, &SMaskListRow::BuildTargetSubmenu));
 					Builder.AddMenuEntry(FGenericCommands::Get().Delete);
+					Builder.AddSubMenu(LOCTEXT("MaskActions_CopyFromMask", "Copy From Mask"), LOCTEXT("MaskActions_CopyFromMask_Tooltip", "Replace this mask with values from another mask on this cloth"), FNewMenuDelegate::CreateSP(this, &SMaskListRow::BuildCopyMaskSubmenu));
 					Builder.AddSubMenu(LOCTEXT("MaskActions_CopyFromVertexColor", "Copy From Vertex Color"), LOCTEXT("MaskActions_CopyFromVertexColor_Tooltip", "Replace this mask with values from vertex color channel on sim mesh"), FNewMenuDelegate::CreateSP(this, &SMaskListRow::BuildCopyVertexColorSubmenu));
 				}
 				Builder.EndSection();
@@ -591,6 +592,38 @@ private:
 		}
 	}
 
+	void OnCopyMask(int32 InMaskIndex)
+	{
+		if (AssetSelectorPtr.IsValid() && Item.IsValid())
+		{
+			if (USkeletalMesh* CurrentMesh = Item->GetOwningMesh())
+			{
+				if (FPointWeightMap* Mask = Item->GetMask())
+				{
+					if (UClothingAssetCommon* const ClothingAsset = AssetSelectorPtr.Pin()->GetSelectedAsset().Get())
+					{
+						const int32 LOD = AssetSelectorPtr.Pin()->GetSelectedLod();
+						if (ClothingAsset->LodData.IsValidIndex(LOD))
+						{
+							const FClothLODDataCommon& LodData = ClothingAsset->LodData[LOD];
+							if (LodData.PointWeightMaps.IsValidIndex(InMaskIndex))
+							{
+								{
+									FScopedTransaction CurrTransaction(LOCTEXT("CopyMask_Transaction", "Copy Mask."));
+									Item->ClothingAsset->Modify();
+									Mask->Values = LodData.PointWeightMaps[InMaskIndex].Values;
+								}
+
+								const bool bUpdateFixedVertData = (Mask->CurrentTarget == (uint8)EWeightMapTargetCommon::MaxDistance);
+								ClothingAsset->ApplyParameterMasks(bUpdateFixedVertData);
+							}
+						}
+					}
+				}
+			}
+		}
+	}
+
 	void BuildTargetSubmenu(FMenuBuilder& Builder)
 	{
 		Builder.BeginSection(NAME_None, LOCTEXT("MaskTargets_SectionName", "Targets"));
@@ -606,7 +639,7 @@ private:
 					// Iterate to -1 to skip the _MAX entry appended to the end of the enum
 					for(int32 Index = 0; Index < NumEntries - 1; ++Index)
 					{
-						FUIAction EntryAction(FExecuteAction::CreateSP(this, &SMaskListRow::OnSetTarget, Index));
+						FUIAction EntryAction(FExecuteAction::CreateSP(this, &SMaskListRow::OnSetTarget, (int32)Enum->GetValueByIndex(Index)));
 
 						FText EntryText = Enum->GetDisplayNameTextByIndex(Index);
 
@@ -618,6 +651,35 @@ private:
 		Builder.EndSection();
 	}
 
+	/** Build sub menu for choosing which vertex color channel to copy to selected mask */
+	void BuildCopyMaskSubmenu(FMenuBuilder& Builder)
+	{
+		if (AssetSelectorPtr.IsValid() && Item.IsValid())
+		{
+			if (const UClothingAssetCommon* const ClothingAsset = AssetSelectorPtr.Pin()->GetSelectedAsset().Get())
+			{
+				const int32 LOD = AssetSelectorPtr.Pin()->GetSelectedLod();
+				if (ClothingAsset->LodData.IsValidIndex(LOD))
+				{
+					const FClothLODDataCommon& LodData = ClothingAsset->LodData[LOD];
+					Builder.BeginSection(NAME_None, LOCTEXT("Masks_SectionName", "Masks"));
+					for (int32 MaskIndex = 0; MaskIndex < LodData.PointWeightMaps.Num(); ++MaskIndex)
+					{
+						if (MaskIndex == Item->MaskIndex)
+						{
+							continue;
+						}
+						FUIAction EntryAction(FExecuteAction::CreateSP(this, &SMaskListRow::OnCopyMask, MaskIndex));
+
+						Builder.AddMenuEntry(FText::FromName(LodData.PointWeightMaps[MaskIndex].Name), FText::GetEmpty(), FSlateIcon(), EntryAction);
+						
+					}
+					Builder.EndSection();
+				}
+			}
+		}
+	}
+
 	/** Build sub menu for choosing which vertex color channel to copy to selected mask */
 	void BuildCopyVertexColorSubmenu(FMenuBuilder& Builder)
 	{
diff --git a/Engine/Source/Runtime/ClothingSystemRuntimeCommon/Private/ClothPhysicalMeshData.cpp b/Engine/Source/Runtime/ClothingSystemRuntimeCommon/Private/ClothPhysicalMeshData.cpp
--- a/Engine/Source/Runtime/ClothingSystemRuntimeCommon/Private/ClothPhysicalMeshData.cpp
+++ b/Engine/Source/Runtime/ClothingSystemRuntimeCommon/Private/ClothPhysicalMeshData.cpp
@@ -237,7 +237,26 @@ void FClothPhysicalMeshData::CalculateNumInfluences()
 
 void FClothPhysicalMeshData::CalculateTethers(bool bUseEuclideanDistance, bool bUseGeodesicDistance)
 {
-	if (const FPointWeightMap* const MaxDistances = FindWeightMap(EWeightMapTargetCommon::MaxDistance))
+	const FPointWeightMap* const MaxDistances = FindWeightMap(EWeightMapTargetCommon::MaxDistance);
+	const FPointWeightMap* const TetherEnds = FindWeightMap(EWeightMapTargetCommon::TetherEndsMask);
+
+	if (TetherEnds)
+	{
+		// Note: Technically there could be two different configs requiring different flavour of tethers,
+		//       in reality at the time of writing there is only one type of tether involved at any one type.
+		if (bUseEuclideanDistance)
+		{
+			constexpr bool bGenerateGeodesic = false;
+			EuclideanTethers.GenerateTethers(Vertices, Indices, TetherEnds->Values, bGenerateGeodesic);
+		}
+		if (bUseGeodesicDistance)
+		{
+			constexpr bool bGenegateGeodesic = true;
+			GeodesicTethers.GenerateTethers(Vertices, Indices, TetherEnds->Values, bGenegateGeodesic);
+		}
+
+	}
+	else if (MaxDistances)
 	{
 		// Note: Technically there could be two different configs requiring different flavour of tethers,
 		//       in reality at the time of writing there is only one type of tether involved at any one type.
diff --git a/Engine/Source/Runtime/ClothingSystemRuntimeCommon/Private/ClothTetherData.cpp b/Engine/Source/Runtime/ClothingSystemRuntimeCommon/Private/ClothTetherData.cpp
--- a/Engine/Source/Runtime/ClothingSystemRuntimeCommon/Private/ClothTetherData.cpp
+++ b/Engine/Source/Runtime/ClothingSystemRuntimeCommon/Private/ClothTetherData.cpp
@@ -15,7 +15,8 @@ public:
 	FClothTetherDataPrivate(
 		const TConstArrayView<FVector3f>& Points,
 		const TConstArrayView<uint32>& Indices,
-		TFunctionRef<bool(int32 Index)> IsKinematic,
+		TFunctionRef<bool(int32 Index)> IsDynamic,
+		TFunctionRef<bool(int32 Index)> IsKinematicTetherEnd,
 		bool bUseGeodesicDistance);
 
 	FClothTetherDataPrivate(
@@ -34,7 +35,7 @@ private:
 
 	// Generate a map of tethers by following the triangle mesh network
 	// The choice of the tether is determined by the shortest geodesic (curvature) distance.
-	void GenerateGeodesicTethers(const TConstArrayView<FVector3f>& Points, TFunctionRef<bool(int32 Index)> IsKinematic);
+	void GenerateGeodesicTethers(const TConstArrayView<FVector3f>& Points, TFunctionRef<bool(int32 Index)> IsKinematicTetherEnd);
 
 	// Update TetherNums after computing tethers.
 	void UpdateCounts();
@@ -62,10 +63,44 @@ void FClothTetherData::GenerateTethers(
 	}
 
 	// Calculate the tethers
-	const FClothTetherDataPrivate ClothTetherData(Points, Indices, [&MaxDistances](int32 Node)
+	const FClothTetherDataPrivate ClothTetherData(Points, Indices,
+		[&MaxDistances](int32 Node)
+		{
+			return MaxDistances[Node] >= FClothTetherDataPrivate::KinematicDistanceThreshold;
+		},
+		[&MaxDistances](int32 Node)
+		{
+			return MaxDistances[Node] < FClothTetherDataPrivate::KinematicDistanceThreshold;
+		},
+		bUseGeodesicDistance);
+	ClothTetherData.GetBatchedTetherData(Tethers);
+}
+
+void FClothTetherData::GenerateTethers(
+	const TConstArrayView<FVector3f>& Points,
+	const TConstArrayView<uint32>& Indices,
+	const TConstArrayView<float>& MaxDistances,
+	const TConstArrayView<float>& TetherEnds,
+	bool bUseGeodesicDistance)
+{
+	// Early exit if there is no MaxDistances mask
+	if (Points.Num() != MaxDistances.Num())
 	{
-		return MaxDistances[Node] < FClothTetherDataPrivate::KinematicDistanceThreshold;
-	}, bUseGeodesicDistance);
+		Tethers.Empty();
+		return;
+	}
+
+	// Calculate the tethers
+	const FClothTetherDataPrivate ClothTetherData(Points, Indices,
+		[&MaxDistances](int32 Node)
+		{
+			return MaxDistances[Node] >= FClothTetherDataPrivate::KinematicDistanceThreshold;
+		},
+		[&TetherEnds](int32 Node)
+		{
+			return TetherEnds[Node] < FClothTetherDataPrivate::KinematicDistanceThreshold;
+		},
+		bUseGeodesicDistance);
 	ClothTetherData.GetBatchedTetherData(Tethers);
 }
 
@@ -83,10 +118,16 @@ void FClothTetherData::GenerateTethers(
 	}
 
 	// Calculate the tethers
-	const FClothTetherDataPrivate ClothTetherData(Points, Indices, [&KinematicNodes](int32 Node)
-	{
-		return KinematicNodes.Contains(Node);
-	}, bUseGeodesicDistance);
+	const FClothTetherDataPrivate ClothTetherData(Points, Indices,
+		[&KinematicNodes](int32 Node)
+		{
+			return !KinematicNodes.Contains(Node);
+		},
+		[&KinematicNodes](int32 Node)
+		{
+			return KinematicNodes.Contains(Node);
+		},
+		bUseGeodesicDistance);
 	ClothTetherData.GetBatchedTetherData(Tethers);
 }
 
@@ -115,7 +156,8 @@ bool FClothTetherData::Serialize(FArchive& Ar)
 FClothTetherDataPrivate::FClothTetherDataPrivate(
 	const TConstArrayView<FVector3f>& Points,
 	const TConstArrayView<uint32>& Indices,
-	TFunctionRef<bool(int32 Index)> IsKinematic,
+	TFunctionRef<bool(int32 Index)> IsDynamic,
+	TFunctionRef<bool(int32 Index)> IsKinematicTetherEnd,
 	bool bUseGeodesicDistance)
 {
 	// Calculate the points' neighbors map
@@ -141,14 +183,15 @@ FClothTetherDataPrivate::FClothTetherDataPrivate(
 	KinematicNodes.Reserve(Points.Num());
 	for (const int32 Node : Nodes)
 	{
-		if (IsKinematic(Node))
+		if (IsDynamic(Node))
 		{
-			KinematicNodes.Add(Node);
+			DynamicNodes.Add(Node);
 		}
-		else
+		else if (IsKinematicTetherEnd(Node))
 		{
-			DynamicNodes.Add(Node);
+			KinematicNodes.Add(Node);
 		}
+		// Skip kinematic nodes that aren't tether ends 
 	}
 
 	if (KinematicNodes.Num())
@@ -171,7 +214,7 @@ FClothTetherDataPrivate::FClothTetherDataPrivate(
 		}
 		else
 		{
-			GenerateGeodesicTethers(Points, IsKinematic);
+			GenerateGeodesicTethers(Points, IsKinematicTetherEnd);
 		}
 
 		UpdateCounts();
@@ -374,11 +417,11 @@ void FClothTetherDataPrivate::GenerateEuclideanTethers(const TConstArrayView<FVe
 	});
 }
 
-void FClothTetherDataPrivate::GenerateGeodesicTethers(const TConstArrayView<FVector3f>& Points, TFunctionRef<bool(int32 Index)> IsKinematic)
+void FClothTetherDataPrivate::GenerateGeodesicTethers(const TConstArrayView<FVector3f>& Points, TFunctionRef<bool(int32 Index)> IsKinematicTetherEnd)
 {
 	check(KinematicNodeIslands.Num());
 
-	// Find all seeds in each island, kinematic nodes connected to at least one dynamic node
+	// Find all seeds in each island, kinematic tether end nodes connected to at least one non-tether end
 	int32 NumSeeds = 0;
 	TArray<TArray<int32>> SeedIslands;
 	SeedIslands.SetNum(KinematicNodeIslands.Num());
@@ -393,7 +436,7 @@ void FClothTetherDataPrivate::GenerateGeodesicTethers(const TConstArrayView<FVec
 			const TSet<int32>& Neighbors = NodeToNeighbors[KinematicNode];
 			for (const int32 Neighbor : Neighbors)
 			{
-				if (!IsKinematic(Neighbor))
+				if (!IsKinematicTetherEnd(Neighbor))
 				{
 					SeedIslands[IslandIndex].Add(KinematicNode);
 					++NumSeeds;
@@ -429,7 +472,7 @@ void FClothTetherDataPrivate::GenerateGeodesicTethers(const TConstArrayView<FVec
 		}
 	}
 
-	ParallelFor(Seeds.Num(), [this, &Seeds, &GeodesicDistances, &Points, &IsKinematic](int32 Index)
+	ParallelFor(Seeds.Num(), [this, &Seeds, &GeodesicDistances, &Points, &IsKinematicTetherEnd](int32 Index)
 	{
 		const int32 Seed = Seeds[Index];
 		TMap<int32, float>& SeedGeodesicDistances = GeodesicDistances[Seed];
@@ -471,8 +514,8 @@ void FClothTetherDataPrivate::GenerateGeodesicTethers(const TConstArrayView<FVec
 			{
 				check(NeighborNode != ParentNode);
 
-				// Do not progress onto kinematic nodes
-				if (IsKinematic(NeighborNode))
+				// Do not progress onto kinematic tether end nodes
+				if (IsKinematicTetherEnd(NeighborNode))
 				{
 					continue;
 				}
diff --git a/Engine/Source/Runtime/ClothingSystemRuntimeCommon/Public/ClothTetherData.h b/Engine/Source/Runtime/ClothingSystemRuntimeCommon/Public/ClothTetherData.h
--- a/Engine/Source/Runtime/ClothingSystemRuntimeCommon/Public/ClothTetherData.h
+++ b/Engine/Source/Runtime/ClothingSystemRuntimeCommon/Public/ClothTetherData.h
@@ -33,7 +33,18 @@ struct FClothTetherData
 		const TConstArrayView<FVector3f>& Points,  // Reference pose
 		const TConstArrayView<uint32>& Indices,  // Triangle mesh
 		const TConstArrayView<float>& MaxDistances,  // Mask for sorting the kinematic from the dynamic points
-		bool bUseGeodesicDistance);  // Whether to use geodesic (walking along the sruface) or euclidean (beeline) distances to find the tethers.
+		bool bUseGeodesicDistance);  // Whether to use geodesic (walking along the surface) or euclidean (beeline) distances to find the tethers.
+	
+	/**
+	 * Generate the tethers by following the triangle mesh network from the closest kinematic tether end to each dynamic point.
+	 * Inside array items can be processed concurrently, but the outside array must be iterated on sequentially.
+	 */
+	CLOTHINGSYSTEMRUNTIMECOMMON_API void GenerateTethers(
+		const TConstArrayView<FVector3f>& Points,  // Reference pose
+		const TConstArrayView<uint32>& Indices,  // Triangle mesh
+		const TConstArrayView<float>& MaxDistances,  // Mask for sorting the kinematic from the dynamic points
+		const TConstArrayView<float>& TetherEnds,  // Separate mask for specifying tether ends (using same < 0.1 cutoff as MaxDistances). Actual tether end set is the intersection of tether ends and kinematic points.
+		bool bUseGeodesicDistance);  // Whether to use geodesic (walking along the surface) or euclidean (beeline) distances to find the tethers.
 	
 	/**
 	 * Generate the tethers by following the triangle mesh network from the closest kinematic to each dynamic point.
@@ -43,7 +54,7 @@ struct FClothTetherData
 		const TConstArrayView<FVector3f>& Points,  // Reference pose
 		const TConstArrayView<uint32>& Indices,  // Triangle mesh
 		const TSet<int32>& KinematicNodes,  // kinematic indices
-		bool bUseGeodesicDistance);  // Whether to use geodesic (walking along the sruface) or euclidean (beeline) distances to find the tethers.
+		bool bUseGeodesicDistance);  // Whether to use geodesic (walking along the surface) or euclidean (beeline) distances to find the tethers.
 
 	/**
 	 * Generate the tether batches from existing tether data.
diff --git a/Engine/Source/Runtime/ClothingSystemRuntimeCommon/Public/PointWeightMap.h b/Engine/Source/Runtime/ClothingSystemRuntimeCommon/Public/PointWeightMap.h
--- a/Engine/Source/Runtime/ClothingSystemRuntimeCommon/Public/PointWeightMap.h
+++ b/Engine/Source/Runtime/ClothingSystemRuntimeCommon/Public/PointWeightMap.h
@@ -15,7 +15,12 @@ enum class EWeightMapTargetCommon : uint8
 	BackstopDistance,    // Distance along the plane of the surface that the particles can travel (separation constraint)
 	BackstopRadius,      // Radius of movement to allow for backstop movement
 	AnimDriveStiffness,  // Strength of anim drive per-particle (spring driving particle back to skinned location
-	AnimDriveDamping_DEPRECATED UMETA(Hidden)  // Chaos onlyweightmap, deprecated from the common declaration
+	AnimDriveDamping_DEPRECATED UMETA(Hidden),  // Chaos onlyweightmap, deprecated from the common declaration
+	FirstUserTarget,
+	// Users can add additional targets between this range.
+	LastUserTarget = 200,
+	// Add additional common targets after here.
+	TetherEndsMask
 };
 
 /** 
