Add Long Range Attachment tether generation sets.

diff --git a/Engine/Plugins/ChaosClothAsset/Source/ChaosClothAssetEngine/Private/ChaosClothAsset/ClothEngineTools.cpp b/Engine/Plugins/ChaosClothAsset/Source/ChaosClothAssetEngine/Private/ChaosClothAsset/ClothEngineTools.cpp
--- a/Engine/Plugins/ChaosClothAsset/Source/ChaosClothAssetEngine/Private/ChaosClothAsset/ClothEngineTools.cpp
+++ b/Engine/Plugins/ChaosClothAsset/Source/ChaosClothAssetEngine/Private/ChaosClothAsset/ClothEngineTools.cpp
@@ -88,4 +88,38 @@ void FClothEngineTools::GenerateTethersFromSelectionSet(const TSharedRef<FManage
 	}
 }
 
+void FClothEngineTools::GenerateTethersFromCustomSelectionSets(const TSharedRef<FManagedArrayCollection>& ClothCollection, const FName& InFixedEndSet, const TArray<TPair<FName, FName>>& CustomTetherEndSets, const bool bGeodesicTethers)
+{
+	FClothGeometryTools::DeleteTethers(ClothCollection);
+	FCollectionClothSelectionConstFacade SelectionFacade(ClothCollection);
+	if (SelectionFacade.HasSelection(InFixedEndSet) && SelectionFacade.GetSelectionGroup(InFixedEndSet) == UE::Chaos::ClothAsset::ClothCollectionGroup::SimVertices3D)
+	{
+
+		FCollectionClothFacade ClothFacade(ClothCollection);
+		TArray<uint32> SimIndices;
+		SimIndices.Reserve(ClothFacade.GetNumSimFaces() * 3);
+		for (const FIntVector3& Face : ClothFacade.GetSimIndices3D())
+		{
+			SimIndices.Add(Face[0]);
+			SimIndices.Add(Face[1]);
+			SimIndices.Add(Face[2]);
+		}
+
+		const TSet<int32>& FixedEndSet = SelectionFacade.GetSelectionSet(InFixedEndSet);
+
+		for (const TPair<FName, FName>& TetherEnds : CustomTetherEndSets)
+		{
+			const FName& CustomDynamicEndSet = TetherEnds.Get<0>();
+			const FName& CustomFixedEndSet = TetherEnds.Get<1>();
+			if (SelectionFacade.HasSelection(CustomFixedEndSet) && SelectionFacade.GetSelectionGroup(CustomFixedEndSet) == UE::Chaos::ClothAsset::ClothCollectionGroup::SimVertices3D &&
+				SelectionFacade.HasSelection(CustomDynamicEndSet) && SelectionFacade.GetSelectionGroup(CustomDynamicEndSet) == UE::Chaos::ClothAsset::ClothCollectionGroup::SimVertices3D)
+			{
+				FClothTetherData TetherData;
+				TetherData.GenerateTethers(ClothFacade.GetSimPosition3D(), TConstArrayView<uint32>(SimIndices), FixedEndSet, SelectionFacade.GetSelectionSet(CustomDynamicEndSet), SelectionFacade.GetSelectionSet(CustomFixedEndSet), bGeodesicTethers);
+				Private::AppendTetherData(ClothFacade, TetherData);
+			}
+		}
+	}
+}
+
 }  // End namespace UE::Chaos::ClothAsset
diff --git a/Engine/Plugins/ChaosClothAsset/Source/ChaosClothAssetEngine/Public/ChaosClothAsset/ClothEngineTools.h b/Engine/Plugins/ChaosClothAsset/Source/ChaosClothAssetEngine/Public/ChaosClothAsset/ClothEngineTools.h
--- a/Engine/Plugins/ChaosClothAsset/Source/ChaosClothAssetEngine/Public/ChaosClothAsset/ClothEngineTools.h
+++ b/Engine/Plugins/ChaosClothAsset/Source/ChaosClothAssetEngine/Public/ChaosClothAsset/ClothEngineTools.h
@@ -19,5 +19,7 @@ namespace UE::Chaos::ClothAsset
 		/** Generate tether data. */
 		static void GenerateTethers(const TSharedRef<FManagedArrayCollection>& ClothCollection, const FName& WeightMap, const bool bGeodesicTethers, const FVector2f& MaxDistanceValue = FVector2f(0.f, 1.f));
 		static void GenerateTethersFromSelectionSet(const TSharedRef<FManagedArrayCollection>& ClothCollection, const FName& FixedEndSet, const bool bGeodesicTethers);
+		/** @param CustomTetherEndSets: First element of each pair is DynamicSet, second is FixedSet */
+		static void GenerateTethersFromCustomSelectionSets(const TSharedRef<FManagedArrayCollection>& ClothCollection, const FName& FixedEndSet, const TArray<TPair<FName, FName>>& CustomTetherEndSets, const bool bGeodesicTethers);
 	};
 }  // End namespace UE::Chaos::ClothAsset
\ No newline at end of file
diff --git a/Engine/Plugins/ChaosClothAssetEditor/Source/ChaosClothAssetDataflowNodes/Private/ChaosClothAsset/SimulationLongRangeAttachmentConfigNode.cpp b/Engine/Plugins/ChaosClothAssetEditor/Source/ChaosClothAssetDataflowNodes/Private/ChaosClothAsset/SimulationLongRangeAttachmentConfigNode.cpp
--- a/Engine/Plugins/ChaosClothAssetEditor/Source/ChaosClothAssetDataflowNodes/Private/ChaosClothAsset/SimulationLongRangeAttachmentConfigNode.cpp
+++ b/Engine/Plugins/ChaosClothAssetEditor/Source/ChaosClothAssetDataflowNodes/Private/ChaosClothAsset/SimulationLongRangeAttachmentConfigNode.cpp
@@ -18,6 +18,14 @@ FChaosClothAssetSimulationLongRangeAttachmentConfigNode_v2::FChaosClothAssetSimu
 	RegisterInputConnection(&TetherScale.WeightMap)
 		.SetCanHidePin(true)
 		.SetPinIsHidden(true);
+
+	// Start with one set of option pins.
+	for (int32 Index = 0; Index < NumInitialCustomTetherSets; ++Index)
+	{
+		AddPins();
+	}
+
+	check(GetNumInputs() == NumRequiredInputs + NumInitialCustomTetherSets * 2); // Update NumRequiredInputs if you add more Inputs. This is used by Serialize.
 }
 
 void FChaosClothAssetSimulationLongRangeAttachmentConfigNode_v2::AddProperties(FPropertyHelper& PropertyHelper) const
@@ -32,7 +40,125 @@ void FChaosClothAssetSimulationLongRangeAttachmentConfigNode_v2::EvaluateClothCo
 {
 	const FString InFixedEndSetString = GetValue<FString>(Context, &FixedEndSet.StringValue);
 	const FName InFixedEndSet(InFixedEndSetString);
-	UE::Chaos::ClothAsset::FClothEngineTools::GenerateTethersFromSelectionSet(ClothCollection, InFixedEndSet, bUseGeodesicTethers);
+	if (bEnableCustomTetherGeneration)
+	{
+		TArray<TPair<FName, FName>> TetherEndSets;
+		TetherEndSets.SetNumUninitialized(CustomTetherData.Num());
+		for (int32 Index = 0; Index < CustomTetherData.Num(); ++Index)
+		{
+			TetherEndSets[Index] = MakeTuple<FName, FName>(FName(*GetValue(Context, GetDynamicEndConnectionReference(Index))), FName(*GetValue(Context, GetFixedEndConnectionReference(Index))));
+		}
+		UE::Chaos::ClothAsset::FClothEngineTools::GenerateTethersFromCustomSelectionSets(ClothCollection, InFixedEndSet, TetherEndSets, bUseGeodesicTethers);
+	}
+	else
+	{
+		UE::Chaos::ClothAsset::FClothEngineTools::GenerateTethersFromSelectionSet(ClothCollection, InFixedEndSet, bUseGeodesicTethers);
+	}
+}
+
+TArray<UE::Dataflow::FPin> FChaosClothAssetSimulationLongRangeAttachmentConfigNode_v2::AddPins()
+{
+	const int32 Index = CustomTetherData.AddDefaulted();
+	TArray<UE::Dataflow::FPin> Pins;
+	Pins.Reserve(2);
+	{
+		const FDataflowInput& Input = RegisterInputArrayConnection(GetFixedEndConnectionReference(Index), GET_MEMBER_NAME_CHECKED(FChaosClothAssetConnectableIStringValue, StringValue));
+		Pins.Emplace(UE::Dataflow::FPin{ UE::Dataflow::FPin::EDirection::INPUT, Input.GetType(), Input.GetName() });
+	}
+	{
+		const FDataflowInput& Input = RegisterInputArrayConnection(GetDynamicEndConnectionReference(Index), GET_MEMBER_NAME_CHECKED(FChaosClothAssetConnectableIStringValue, StringValue));
+		Pins.Emplace(UE::Dataflow::FPin{ UE::Dataflow::FPin::EDirection::INPUT, Input.GetType(), Input.GetName() });
+	}
+	return Pins;
+}
+
+TArray<UE::Dataflow::FPin> FChaosClothAssetSimulationLongRangeAttachmentConfigNode_v2::GetPinsToRemove() const
+{
+	const int32 Index = CustomTetherData.Num() - 1;
+	check(CustomTetherData.IsValidIndex(Index));
+	TArray<UE::Dataflow::FPin> Pins;
+	Pins.Reserve(2);
+	if (const FDataflowInput* const Input = FindInput(GetFixedEndConnectionReference(Index)))
+	{
+		Pins.Emplace(UE::Dataflow::FPin{ UE::Dataflow::FPin::EDirection::INPUT, Input->GetType(), Input->GetName() });
+	}
+	if (const FDataflowInput* const Input = FindInput(GetDynamicEndConnectionReference(Index)))
+	{
+		Pins.Emplace(UE::Dataflow::FPin{ UE::Dataflow::FPin::EDirection::INPUT, Input->GetType(), Input->GetName() });
+	}
+	return Pins;
+}
+
+void FChaosClothAssetSimulationLongRangeAttachmentConfigNode_v2::OnPinRemoved(const UE::Dataflow::FPin& Pin)
+{
+	const int32 Index = CustomTetherData.Num() - 1;
+	check(CustomTetherData.IsValidIndex(Index));
+	const FDataflowInput* const FirstInput = FindInput(GetFixedEndConnectionReference(Index));
+	const FDataflowInput* const SecondInput = FindInput(GetDynamicEndConnectionReference(Index));
+	check(FirstInput || SecondInput);
+	const bool bIsFirstInput = FirstInput && FirstInput->GetName() == Pin.Name;
+	const bool bIsSecondInput = SecondInput && SecondInput->GetName() == Pin.Name;
+	if ((bIsFirstInput && !SecondInput) || (bIsSecondInput && !FirstInput))
+	{
+		// Both inputs removed. Remove array index.
+		CustomTetherData.SetNum(Index);
+	}
+	return Super::OnPinRemoved(Pin);
+}
+
+void FChaosClothAssetSimulationLongRangeAttachmentConfigNode_v2::PostSerialize(const FArchive& Ar)
+{
+	// Restore the pins when re-loading so they can get properly reconnected
+	if (Ar.IsLoading())
+	{
+		check(CustomTetherData.Num() >= NumInitialCustomTetherSets);
+		for (int32 Index = 0; Index < NumInitialCustomTetherSets; ++Index)
+		{
+			check(FindInput(GetFixedEndConnectionReference(Index)));
+			check(FindInput(GetDynamicEndConnectionReference(Index)));
+		}
+
+		for (int32 Index = NumInitialCustomTetherSets; Index < CustomTetherData.Num(); ++Index)
+		{
+			FindOrRegisterInputArrayConnection(GetFixedEndConnectionReference(Index), GET_MEMBER_NAME_CHECKED(FChaosClothAssetConnectableIStringValue, StringValue));
+			FindOrRegisterInputArrayConnection(GetDynamicEndConnectionReference(Index), GET_MEMBER_NAME_CHECKED(FChaosClothAssetConnectableIStringValue, StringValue));
+		}
+
+		if (Ar.IsTransacting())
+		{
+			const int32 OrigNumRegisteredInputs = GetNumInputs();
+			check(OrigNumRegisteredInputs >= NumRequiredInputs + NumInitialCustomTetherSets * 2);
+			const int32 OrigNumSets = CustomTetherData.Num();
+			const int32 OrigNumRegisteredSets = (OrigNumRegisteredInputs - NumRequiredInputs) / 2;
+
+			if (OrigNumRegisteredSets > OrigNumSets)
+			{
+				ensure(Ar.IsTransacting());
+				// Temporarily expand SelectionFilterSets so we can get connection references.
+				CustomTetherData.SetNum(OrigNumRegisteredSets);
+				for (int32 Index = OrigNumSets; Index < CustomTetherData.Num(); ++Index)
+				{
+					UnregisterInputConnection(GetDynamicEndConnectionReference(Index));
+					UnregisterInputConnection(GetFixedEndConnectionReference(Index));
+				}
+				CustomTetherData.SetNum(OrigNumRegisteredSets);
+			}
+		}
+		else
+		{
+			ensureAlways(CustomTetherData.Num() * 2 + NumRequiredInputs == GetNumInputs());
+		}
+	}
+}
+
+UE::Dataflow::TConnectionReference<FString> FChaosClothAssetSimulationLongRangeAttachmentConfigNode_v2::GetFixedEndConnectionReference(int32 Index) const
+{
+	return { &CustomTetherData[Index].CustomFixedEndSet.StringValue, Index, &CustomTetherData };
+}
+
+UE::Dataflow::TConnectionReference<FString> FChaosClothAssetSimulationLongRangeAttachmentConfigNode_v2::GetDynamicEndConnectionReference(int32 Index) const
+{
+	return { &CustomTetherData[Index].CustomDynamicEndSet.StringValue, Index, &CustomTetherData };
 }
 
 
diff --git a/Engine/Plugins/ChaosClothAssetEditor/Source/ChaosClothAssetDataflowNodes/Public/ChaosClothAsset/SimulationLongRangeAttachmentConfigNode.h b/Engine/Plugins/ChaosClothAssetEditor/Source/ChaosClothAssetDataflowNodes/Public/ChaosClothAsset/SimulationLongRangeAttachmentConfigNode.h
--- a/Engine/Plugins/ChaosClothAssetEditor/Source/ChaosClothAssetDataflowNodes/Public/ChaosClothAsset/SimulationLongRangeAttachmentConfigNode.h
+++ b/Engine/Plugins/ChaosClothAssetEditor/Source/ChaosClothAssetDataflowNodes/Public/ChaosClothAsset/SimulationLongRangeAttachmentConfigNode.h
@@ -6,6 +6,24 @@
 #include "ChaosClothAsset/WeightedValue.h"
 #include "SimulationLongRangeAttachmentConfigNode.generated.h"
 
+USTRUCT()
+struct FChaosClothAssetTetherGenerationSet
+{
+	GENERATED_BODY()
+
+	/**
+	* Set of Kinematic vertices that tethers will attach to. Will be intersected with (node level) FixedEndSet.
+	*/
+	UPROPERTY(EditAnywhere, Category = "Long Range Attachment Properties")
+	FChaosClothAssetConnectableIStringValue CustomFixedEndSet = { TEXT("KinematicVertices3D") };
+
+	/**
+	* Set of Dynamic vertices that will attach to the closest vertex in FixedEndSelection. (Node level) FixedEndSet will be excluded. 
+	*/
+	UPROPERTY(EditAnywhere, Category = "Long Range Attachment Properties")
+	FChaosClothAssetConnectableIStringValue CustomDynamicEndSet = { TEXT("DynamicVertices3D") };
+};
+
 /** Long range attachment constraint property configuration node. */
 USTRUCT(Meta = (DataflowCloth))
 struct FChaosClothAssetSimulationLongRangeAttachmentConfigNode_v2 : public FChaosClothAssetSimulationBaseConfigNode
@@ -44,16 +62,42 @@ public:
 	UPROPERTY(EditAnywhere, Category = "Long Range Attachment Properties")
 	bool bUseGeodesicTethers = true;
 
-	/** The name of the vertex selection set used as fixed tether ends. */
+	/**
+	* Enable more granular control over tether generation via custom selection sets.
+	* Otherwise, all dynamic particles will be connect to the closest kinematic vertex as defined by FixedEndSet.
+	*/
+	UPROPERTY(EditAnywhere, Category = "Long Range Attachment Properties")
+	bool bEnableCustomTetherGeneration = false;
+
+	/** The name of the vertex selection set used as fixed tether ends.
+	* When using custom tether generation, this set is still needed to contain all kinematic vertices.
+	*/
 	UPROPERTY(EditAnywhere, Category = "Long Range Attachment Properties")
 	FChaosClothAssetConnectableIStringValue FixedEndSet = { TEXT("KinematicVertices3D") };
 
+	/** Pairs of vertex selections used for custom tether generation. */
+	UPROPERTY(EditAnywhere, EditFixedSize, Category = "Long Range Attachment Properties", Meta = (EditCondition = "bEnableCustomTetherGeneration", SkipInDisplayNameChain))
+	TArray<FChaosClothAssetTetherGenerationSet> CustomTetherData;
+
 	FChaosClothAssetSimulationLongRangeAttachmentConfigNode_v2(const UE::Dataflow::FNodeParameters& InParam, FGuid InGuid = FGuid::NewGuid());
 
 private:
 	virtual void AddProperties(FPropertyHelper& PropertyHelper) const override;
 
 	virtual void EvaluateClothCollection(UE::Dataflow::FContext& Context, const TSharedRef<FManagedArrayCollection>& ClothCollection) const override;
+
+	virtual TArray<UE::Dataflow::FPin> AddPins() override;
+	virtual bool CanAddPin() const override { return true; }
+	virtual bool CanRemovePin() const override { return CustomTetherData.Num() > NumInitialCustomTetherSets; }
+	virtual TArray<UE::Dataflow::FPin> GetPinsToRemove() const override;
+	virtual void OnPinRemoved(const UE::Dataflow::FPin& Pin) override;
+	virtual void PostSerialize(const FArchive& Ar) override;
+
+	UE::Dataflow::TConnectionReference<FString> GetFixedEndConnectionReference(int32 Index) const;
+	UE::Dataflow::TConnectionReference<FString> GetDynamicEndConnectionReference(int32 Index) const;
+
+	static constexpr int32 NumRequiredInputs = 4; // non-filter set inputs
+	static constexpr int32 NumInitialCustomTetherSets = 1;
 };
 
 /** Long range attachment constraint property configuration node. */
diff --git a/Engine/Source/Runtime/ClothingSystemRuntimeCommon/Private/ClothTetherData.cpp b/Engine/Source/Runtime/ClothingSystemRuntimeCommon/Private/ClothTetherData.cpp
--- a/Engine/Source/Runtime/ClothingSystemRuntimeCommon/Private/ClothTetherData.cpp
+++ b/Engine/Source/Runtime/ClothingSystemRuntimeCommon/Private/ClothTetherData.cpp
@@ -15,7 +15,8 @@ public:
 	FClothTetherDataPrivate(
 		const TConstArrayView<FVector3f>& Points,
 		const TConstArrayView<uint32>& Indices,
-		TFunctionRef<bool(int32 Index)> IsDynamic,
+		TFunctionRef<bool(int32 Index)> IsKinematic,
+		TFunctionRef<bool(int32 Index)> IsDynamicTetherEnd,
 		TFunctionRef<bool(int32 Index)> IsKinematicTetherEnd,
 		bool bUseGeodesicDistance);
 
@@ -35,15 +36,17 @@ private:
 
 	// Generate a map of tethers by following the triangle mesh network
 	// The choice of the tether is determined by the shortest geodesic (curvature) distance.
-	void GenerateGeodesicTethers(const TConstArrayView<FVector3f>& Points, TFunctionRef<bool(int32 Index)> IsKinematicTetherEnd);
+	void GenerateGeodesicTethers(const TConstArrayView<FVector3f>& Points,
+		TFunctionRef<bool(int32 Index)> IsKinematic);
 
 	// Update TetherNums after computing tethers.
 	void UpdateCounts();
 
 private:
 	TMap<int32, TSet<int32>> NodeToNeighbors;
-	TArray<int32> DynamicNodes;
-	TArray<int32> KinematicNodes;
+	TArray<int32> AllDynamicNodes;
+	TArray<int32> DynamicTetherNodesDynamicIndex; // Index into AllDynamicNodes, not Nodes.
+	TArray<int32> KinematicTetherNodes;
 	TArray<TArray<int32>> KinematicNodeIslands;
 	TArray<TArray<FTether>> TetherSlots;  // Each dynamic node has between 0 and MaxNumAttachments slots
 	TArray<int32> TetherNums;
@@ -64,6 +67,10 @@ void FClothTetherData::GenerateTethers(
 
 	// Calculate the tethers
 	const FClothTetherDataPrivate ClothTetherData(Points, Indices,
+		[&MaxDistances](int32 Node)
+		{
+			return MaxDistances[Node] < FClothTetherDataPrivate::KinematicDistanceThreshold;
+		},
 		[&MaxDistances](int32 Node)
 		{
 			return MaxDistances[Node] >= FClothTetherDataPrivate::KinematicDistanceThreshold;
@@ -92,6 +99,10 @@ void FClothTetherData::GenerateTethers(
 
 	// Calculate the tethers
 	const FClothTetherDataPrivate ClothTetherData(Points, Indices,
+		[&MaxDistances](int32 Node)
+		{
+			return MaxDistances[Node] < FClothTetherDataPrivate::KinematicDistanceThreshold;
+		},
 		[&MaxDistances](int32 Node)
 		{
 			return MaxDistances[Node] >= FClothTetherDataPrivate::KinematicDistanceThreshold;
@@ -119,6 +130,10 @@ void FClothTetherData::GenerateTethers(
 
 	// Calculate the tethers
 	const FClothTetherDataPrivate ClothTetherData(Points, Indices,
+		[&KinematicNodes](int32 Node)
+		{
+			return KinematicNodes.Contains(Node);
+		},
 		[&KinematicNodes](int32 Node)
 		{
 			return !KinematicNodes.Contains(Node);
@@ -131,6 +146,40 @@ void FClothTetherData::GenerateTethers(
 	ClothTetherData.GetBatchedTetherData(Tethers);
 }
 
+void FClothTetherData::GenerateTethers(
+	const TConstArrayView<FVector3f>& Points,  // Reference pose
+	const TConstArrayView<uint32>& Indices,  // Triangle mesh
+	const TSet<int32>& KinematicNodes,  // all kinematic indices
+	const TSet<int32>& FixedEndNodes,  // subset of kinematic indices used for fixed end.
+	const TSet<int32>& DynamicEndNodes,  // subset of dynamic indices to create tethers for.
+	bool bUseGeodesicDistance)
+{
+	// Early exit if there is no KinematicNodes or DynamicNodes set
+	if (FixedEndNodes.IsEmpty() || DynamicEndNodes.IsEmpty() || KinematicNodes.IsEmpty())
+	{
+		Tethers.Empty();
+		return;
+	}
+
+	// Calculate the tethers
+	const FClothTetherDataPrivate ClothTetherData(Points, Indices,
+		[&KinematicNodes](int32 Node)
+		{
+			return KinematicNodes.Contains(Node);
+		},
+		[&KinematicNodes, &DynamicEndNodes](int32 Node)
+		{
+			return DynamicEndNodes.Contains(Node) && !KinematicNodes.Contains(Node);
+		},
+		[&KinematicNodes, &FixedEndNodes](int32 Node)
+		{
+			return FixedEndNodes.Contains(Node) && KinematicNodes.Contains(Node);
+		},
+		bUseGeodesicDistance);
+	ClothTetherData.GetBatchedTetherData(Tethers);
+
+}
+
 void FClothTetherData::GenerateTethers(
 	TArray<TArray<TPair<float, int32>>>&& PerDynamicNodeTethers)
 {
@@ -156,7 +205,8 @@ bool FClothTetherData::Serialize(FArchive& Ar)
 FClothTetherDataPrivate::FClothTetherDataPrivate(
 	const TConstArrayView<FVector3f>& Points,
 	const TConstArrayView<uint32>& Indices,
-	TFunctionRef<bool(int32 Index)> IsDynamic,
+	TFunctionRef<bool(int32 Index)> IsKinematic,
+	TFunctionRef<bool(int32 Index)> IsDynamicTetherEnd,
 	TFunctionRef<bool(int32 Index)> IsKinematicTetherEnd,
 	bool bUseGeodesicDistance)
 {
@@ -179,28 +229,33 @@ FClothTetherDataPrivate::FClothTetherDataPrivate(
 	NodeToNeighbors.GenerateKeyArray(Nodes);
 
 	// Find all kinematic points to use as anchor points for the tethers and seed the path finding
-	DynamicNodes.Reserve(Points.Num());
-	KinematicNodes.Reserve(Points.Num());
+	AllDynamicNodes.Reserve(Points.Num());
+	DynamicTetherNodesDynamicIndex.Reserve(Points.Num());
+	KinematicTetherNodes.Reserve(Points.Num());
 	for (const int32 Node : Nodes)
 	{
-		if (IsDynamic(Node))
+		if (!IsKinematic(Node))
 		{
-			DynamicNodes.Add(Node);
+			const int32 DynamicIndex = AllDynamicNodes.Add(Node);
+			if (IsDynamicTetherEnd(Node))
+			{
+				DynamicTetherNodesDynamicIndex.Add(DynamicIndex);
+			}
+
 		}
 		else if (IsKinematicTetherEnd(Node))
 		{
-			KinematicNodes.Add(Node);
+			KinematicTetherNodes.Add(Node);
 		}
-		// Skip kinematic nodes that aren't tether ends 
 	}
 
-	if (KinematicNodes.Num())
+	if (KinematicTetherNodes.Num())
 	{
 		// Compute the islands of kinematic particles
 		ComputeKinematicNodeIslands();
 
 		// Allocate the tether batches, each node has a maximum of MaxNumAttachments slots
-		TetherSlots.SetNum(DynamicNodes.Num());
+		TetherSlots.SetNum(DynamicTetherNodesDynamicIndex.Num());
 
 		for (TArray<FTether>& TetherSlot : TetherSlots)
 		{
@@ -214,7 +269,7 @@ FClothTetherDataPrivate::FClothTetherDataPrivate(
 		}
 		else
 		{
-			GenerateGeodesicTethers(Points, IsKinematicTetherEnd);
+			GenerateGeodesicTethers(Points, IsKinematic);
 		}
 
 		UpdateCounts();
@@ -284,9 +339,9 @@ void FClothTetherDataPrivate::ComputeKinematicNodeIslands()
 	TArray<int32> FreeIslandIndices;
 
 	TMap<int32, int32> NodeToIslands;
-	NodeToIslands.Reserve(KinematicNodes.Num());
+	NodeToIslands.Reserve(KinematicTetherNodes.Num());
 
-	for (const int32 KinematicNode : KinematicNodes)
+	for (const int32 KinematicNode : KinematicTetherNodes)
 	{
 		// Assign this point an island, possibly unionizing existing islands
 		int32 IslandIndex = TNumericLimits<int32>::Max();
@@ -358,7 +413,7 @@ void FClothTetherDataPrivate::ComputeKinematicNodeIslands()
 	}
 
 	// Final sanity check
-	check(TotalNumIslandKinematicNodes == KinematicNodes.Num());
+	check(TotalNumIslandKinematicNodes == KinematicTetherNodes.Num());
 }
 
 void FClothTetherDataPrivate::GenerateEuclideanTethers(const TConstArrayView<FVector3f>& Points)
@@ -366,9 +421,9 @@ void FClothTetherDataPrivate::GenerateEuclideanTethers(const TConstArrayView<FVe
 	check(KinematicNodeIslands.Num());
 
 	// For each dynamic node
-	ParallelFor(DynamicNodes.Num(), [this, &Points](int32 Index)
+	ParallelFor(DynamicTetherNodesDynamicIndex.Num(), [this, &Points](int32 Index)
 	{
-		const int32 DynamicNode = DynamicNodes[Index];
+		const int32 DynamicNode = AllDynamicNodes[DynamicTetherNodesDynamicIndex[Index]];
 
 		// Measure the distance to all kinematic nodes, and keep the closest from each island
 		TArray<TPair<float, int32>> ClosestKinematicNodes;
@@ -417,7 +472,8 @@ void FClothTetherDataPrivate::GenerateEuclideanTethers(const TConstArrayView<FVe
 	});
 }
 
-void FClothTetherDataPrivate::GenerateGeodesicTethers(const TConstArrayView<FVector3f>& Points, TFunctionRef<bool(int32 Index)> IsKinematicTetherEnd)
+void FClothTetherDataPrivate::GenerateGeodesicTethers(const TConstArrayView<FVector3f>& Points,
+	TFunctionRef<bool(int32 Index)> IsKinematic)
 {
 	check(KinematicNodeIslands.Num());
 
@@ -436,7 +492,7 @@ void FClothTetherDataPrivate::GenerateGeodesicTethers(const TConstArrayView<FVec
 			const TSet<int32>& Neighbors = NodeToNeighbors[KinematicNode];
 			for (const int32 Neighbor : Neighbors)
 			{
-				if (!IsKinematicTetherEnd(Neighbor))
+				if (!IsKinematic(Neighbor))
 				{
 					SeedIslands[IslandIndex].Add(KinematicNode);
 					++NumSeeds;
@@ -464,15 +520,15 @@ void FClothTetherDataPrivate::GenerateGeodesicTethers(const TConstArrayView<FVec
 	for (const int32 Seed : Seeds)
 	{
 		TMap<int32, float>& SeedGeodesicDistances = GeodesicDistances.Emplace(Seed);
-		SeedGeodesicDistances.Reserve(DynamicNodes.Num());
+		SeedGeodesicDistances.Reserve(AllDynamicNodes.Num());
 
-		for (const int32 DynamicNode : DynamicNodes)
+		for (const int32 DynamicNode : AllDynamicNodes)
 		{
 			SeedGeodesicDistances.Emplace(DynamicNode, TNumericLimits<float>::Max());
 		}
 	}
 
-	ParallelFor(Seeds.Num(), [this, &Seeds, &GeodesicDistances, &Points, &IsKinematicTetherEnd](int32 Index)
+	ParallelFor(Seeds.Num(), [this, &Seeds, &GeodesicDistances, &Points, &IsKinematic](int32 Index)
 	{
 		const int32 Seed = Seeds[Index];
 		TMap<int32, float>& SeedGeodesicDistances = GeodesicDistances[Seed];
@@ -514,8 +570,8 @@ void FClothTetherDataPrivate::GenerateGeodesicTethers(const TConstArrayView<FVec
 			{
 				check(NeighborNode != ParentNode);
 
-				// Do not progress onto kinematic tether end nodes
-				if (IsKinematicTetherEnd(NeighborNode))
+				// Do not progress onto kinematic nodes
+				if (IsKinematic(NeighborNode))
 				{
 					continue;
 				}
@@ -547,7 +603,7 @@ void FClothTetherDataPrivate::GenerateGeodesicTethers(const TConstArrayView<FVec
 	});
 
 	// Initialize the tether slot arrays
-	TetherSlots.SetNum(DynamicNodes.Num());
+	TetherSlots.SetNum(DynamicTetherNodesDynamicIndex.Num());
 
 	for (TArray<FTether>& NewTethers : TetherSlots)
 	{
@@ -555,9 +611,9 @@ void FClothTetherDataPrivate::GenerateGeodesicTethers(const TConstArrayView<FVec
 	}
 
 	// Find the tether constraints starting from each dynamic node
-	ParallelFor(DynamicNodes.Num(), [this, &SeedIslands, &GeodesicDistances](int32 Index)
+	ParallelFor(DynamicTetherNodesDynamicIndex.Num(), [this, &SeedIslands, &GeodesicDistances](int32 Index)
 	{
-		const int32 DynamicNode = DynamicNodes[Index];
+		const int32 DynamicNode = AllDynamicNodes[DynamicTetherNodesDynamicIndex[Index]];
 
 		// Find the closest seeds in each island
 		TArray<int32, TInlineAllocator<32>> ClosestSeeds;
diff --git a/Engine/Source/Runtime/ClothingSystemRuntimeCommon/Public/ClothTetherData.h b/Engine/Source/Runtime/ClothingSystemRuntimeCommon/Public/ClothTetherData.h
--- a/Engine/Source/Runtime/ClothingSystemRuntimeCommon/Public/ClothTetherData.h
+++ b/Engine/Source/Runtime/ClothingSystemRuntimeCommon/Public/ClothTetherData.h
@@ -55,6 +55,18 @@ struct FClothTetherData
 		const TConstArrayView<uint32>& Indices,  // Triangle mesh
 		const TSet<int32>& KinematicNodes,  // kinematic indices
 		bool bUseGeodesicDistance);  // Whether to use geodesic (walking along the surface) or euclidean (beeline) distances to find the tethers.
+	
+	/**
+	 * Generate the tethers by following the triangle mesh network from the closest kinematic to each dynamic point.
+	 * Inside array items can be processed concurrently, but the outside array must be iterated on sequentially.
+	 */
+	CLOTHINGSYSTEMRUNTIMECOMMON_API void GenerateTethers(
+		const TConstArrayView<FVector3f>& Points,  // Reference pose
+		const TConstArrayView<uint32>& Indices,  // Triangle mesh
+		const TSet<int32>& KinematicNodes,  // all kinematic indices
+		const TSet<int32>& FixedEndNodes,  // subset of kinematic indices used for fixed end.
+		const TSet<int32>& DynamicEndNodes,  // subset of dynamic indices to create tethers for.
+		bool bUseGeodesicDistance);  // Whether to use geodesic (walking along the surface) or euclidean (beeline) distances to find the tethers.
 
 	/**
 	 * Generate the tether batches from existing tether data.
