Fix CorrespondClothAssetIndex never set in BindToSkeletalMesh, since checks are stripped out in Shipping

diff --git a/Engine/Source/Runtime/ClothingSystemRuntimeCommon/Private/ClothingAsset.cpp b/Engine/Source/Runtime/ClothingSystemRuntimeCommon/Private/ClothingAsset.cpp
--- a/Engine/Source/Runtime/ClothingSystemRuntimeCommon/Private/ClothingAsset.cpp
+++ b/Engine/Source/Runtime/ClothingSystemRuntimeCommon/Private/ClothingAsset.cpp
@@ -309,7 +309,8 @@ bool UClothingAssetCommon::BindToSkeletalMesh(
 
 	// Set the asset index, used during rendering to pick the correct sim mesh buffer
 	int32 AssetIndex = INDEX_NONE;
-	check(InSkelMesh->GetMeshClothingAssets().Find(this, AssetIndex));
+	const bool bMeshClothingAssetsWasFound = InSkelMesh->GetMeshClothingAssets().Find(this, AssetIndex);
+	check(bMeshClothingAssetsWasFound);
 	OriginalSection.CorrespondClothAssetIndex = AssetIndex;
 
 	// sim properties
