Skeletal Mesh Cloth - Adjusted a slight change of backward compatibility behavior with very old clothing assets caused by a recent change in CL 43909990.

#rnx

diff --git a/Engine/Source/Runtime/ClothingSystemRuntimeCommon/Private/ClothingAsset.cpp b/Engine/Source/Runtime/ClothingSystemRuntimeCommon/Private/ClothingAsset.cpp
--- a/Engine/Source/Runtime/ClothingSystemRuntimeCommon/Private/ClothingAsset.cpp
+++ b/Engine/Source/Runtime/ClothingSystemRuntimeCommon/Private/ClothingAsset.cpp
@@ -1084,7 +1084,7 @@ void UClothingAssetCommon::PostLoad()
 		// and to unify all cached data invalidations, while the InvalidateCachedData method is also deprecated from ClothingAssetBase.
 		if (SkeletalMesh->GetLinkerCustomVersion(FSkeletalMeshCustomVersion::GUID) < FSkeletalMeshCustomVersion::CachedClothInverseMasses)
 		{
-			ClothingCachedDataFlags |= EClothingCachedDataFlagsCommon::InverseMasses;
+			ClothingCachedDataFlags |= EClothingCachedDataFlagsCommon::All;
 		}
 	}
 
