Cloth - Rename Skeletal Mesh Clothing Data Config category into "Skeletal Mesh Clothing".

#rnx

diff --git a/Engine/Source/Runtime/ClothingSystemRuntimeCommon/Public/ClothingAsset.h b/Engine/Source/Runtime/ClothingSystemRuntimeCommon/Public/ClothingAsset.h
--- a/Engine/Source/Runtime/ClothingSystemRuntimeCommon/Public/ClothingAsset.h
+++ b/Engine/Source/Runtime/ClothingSystemRuntimeCommon/Public/ClothingAsset.h
@@ -212,12 +212,12 @@ public:
 	CLOTHINGSYSTEMRUNTIMECOMMON_API virtual void PostUpdateAllAssets() override;
 
 	// The physics asset to extract collisions from when building a simulation.
-	UPROPERTY(EditAnywhere, Category = Config)
+	UPROPERTY(EditAnywhere, Category = "Skeletal Mesh Clothing")
 	TObjectPtr<UPhysicsAsset> PhysicsAsset;
 
 	// Simulation specific cloth parameters. 
 	// Use GetClothConfig() to retrieve the correct parameters/config type for the desired cloth simulation system.
-	UPROPERTY(VisibleAnywhere, BlueprintReadOnly, EditFixedSize, Instanced, Category = Config)
+	UPROPERTY(VisibleAnywhere, BlueprintReadOnly, EditFixedSize, Instanced, Category = "Skeletal Mesh Clothing")
 	TMap<FName, TObjectPtr<UClothConfigBase>> ClothConfigs;
 
 #if WITH_EDITORONLY_DATA
