[Backout] - CL44268492
#rnx
Original CL Desc
-----------------------------------------------------------------
Chaos Cloth Asset - Added first pass factories to use the Cloth Asset inside the SkeletalMesh.

#rnx

diff --git a/Engine/Plugins/ChaosClothAsset/Source/ChaosClothAssetEngine/ChaosClothAssetEngine.Build.cs b/Engine/Plugins/ChaosClothAsset/Source/ChaosClothAssetEngine/ChaosClothAssetEngine.Build.cs
--- a/Engine/Plugins/ChaosClothAsset/Source/ChaosClothAssetEngine/ChaosClothAssetEngine.Build.cs
+++ b/Engine/Plugins/ChaosClothAsset/Source/ChaosClothAssetEngine/ChaosClothAssetEngine.Build.cs
@@ -29,8 +29,7 @@ public class ChaosClothAssetEngine : ModuleRules
 				"ChaosClothAsset",
 				"ChaosCloth",
 				"ChaosCaching",
-				"ClothingSystemRuntimeInterface",
-				"ClothingSystemEditorInterface"  // For UChaosClothAssetSKMClothingAssetFactory
+				"ClothingSystemRuntimeInterface"
 			}
 		);
 		SetupModulePhysicsSupport(Target);
diff --git a/Engine/Plugins/ChaosClothAsset/Source/ChaosClothAssetEngine/Private/ChaosClothAsset/ChaosClothAssetEngineModule.cpp b/Engine/Plugins/ChaosClothAsset/Source/ChaosClothAssetEngine/Private/ChaosClothAsset/ChaosClothAssetEngineModule.cpp
--- a/Engine/Plugins/ChaosClothAsset/Source/ChaosClothAssetEngine/Private/ChaosClothAsset/ChaosClothAssetEngineModule.cpp
+++ b/Engine/Plugins/ChaosClothAsset/Source/ChaosClothAssetEngine/Private/ChaosClothAsset/ChaosClothAssetEngineModule.cpp
@@ -1,15 +1,11 @@
 // Copyright Epic Games, Inc. All Rights Reserved.
 
 #include "ChaosClothAsset/ClothAssetPrivate.h"
-#include "ChaosClothAsset/ClothAssetSKMSimulationFactory.h"
+#include "ChaosClothAsset/ClothComponent.h"
 #include "ChaosClothAsset/ClothComponentCacheAdapter.h"
-#include "Features/IModularFeatures.h"
 #include "Modules/ModuleInterface.h"
 #include "Modules/ModuleManager.h"
 #if WITH_EDITOR
-#include "ChaosClothAsset/ClothAssetSKMClothingAssetFactory.h"
-#include "ChaosClothAsset/ClothComponent.h"
-#include "ClothingAssetFactoryInterface.h"
 #include "PropertyEditorModule.h"
 #endif
 
@@ -19,15 +15,9 @@ DEFINE_LOG_CATEGORY(LogChaosClothAsset);
 
 namespace UE::Chaos::ClothAsset
 {
-	class FClothAssetEngineModule
-		: public IModuleInterface
-		, public IClothingSimulationFactoryClassProvider
-#if WITH_EDITOR
-		, public IClothingAssetFactoryProvider
-#endif
+	class FClothAssetEngineModule : public IModuleInterface
 	{
 	public:
-		//~ Begin IModuleInterface
 		virtual void StartupModule() override
 		{
 			ClothComponentCacheAdapter = MakeUnique<FClothComponentCacheAdapter>();
@@ -39,18 +29,12 @@ namespace UE::Chaos::ClothAsset
 			const TSharedPtr<FPropertySection> Section = PropertyModule.FindOrCreateSection(UChaosClothComponent::StaticClass()->GetFName(), TEXT("ClothSim"), LOCTEXT("ClothSim", "Cloth Sim"));
 			Section->AddCategory("Cloth Component");
 			PropertyModule.NotifyCustomizationModuleChanged();
-
-			IModularFeatures::Get().RegisterModularFeature(IClothingAssetFactoryProvider::FeatureName, static_cast<IClothingAssetFactoryProvider*>(this));
 #endif
-			IModularFeatures::Get().RegisterModularFeature(IClothingSimulationFactoryClassProvider::FeatureName, static_cast<IClothingSimulationFactoryClassProvider*>(this));
 		}
 
 		virtual void ShutdownModule() override
 		{
-			IModularFeatures::Get().UnregisterModularFeature(IClothingSimulationFactoryClassProvider::FeatureName, static_cast<IClothingSimulationFactoryClassProvider*>(this));
 #if WITH_EDITOR
-			IModularFeatures::Get().UnregisterModularFeature(IClothingAssetFactoryProvider::FeatureName, static_cast<IClothingAssetFactoryProvider*>(this));
-
 			if (UObjectInitialized())
 			{
 				// Remove sections
@@ -59,24 +43,6 @@ namespace UE::Chaos::ClothAsset
 			}
 #endif
 		}
-		//~ End IModuleInterface
-
-		//~ Begin IClothingSimulationFactoryClassProvider
-		virtual TSubclassOf<class UClothingSimulationFactory> GetClothingSimulationFactoryClass() const override
-		{
-			return UChaosClothAssetSKMSimulationFactory::StaticClass();
-		}
-		//~ End IClothingSimulationFactoryClassProvider
-
-#if WITH_EDITOR
-		//~ Begin IClothingAssetFactoryProvider
-		virtual UClothingAssetFactoryBase* GetFactory() override
-		{
-			return UChaosClothAssetSKMClothingAssetFactory::StaticClass()->GetDefaultObject<UClothingAssetFactoryBase>();
-		}
-		//~ End IClothingSimulationFactoryClassProvider
-#endif
-
 	private:
 		TUniquePtr<FClothComponentCacheAdapter> ClothComponentCacheAdapter;
 	};
diff --git a/Engine/Plugins/ChaosClothAsset/Source/ChaosClothAssetEngine/Private/ChaosClothAsset/ClothAssetBuilder.cpp b/Engine/Plugins/ChaosClothAsset/Source/ChaosClothAssetEngine/Private/ChaosClothAsset/ClothAssetBuilder.cpp
--- a/Engine/Plugins/ChaosClothAsset/Source/ChaosClothAssetEngine/Private/ChaosClothAsset/ClothAssetBuilder.cpp
+++ b/Engine/Plugins/ChaosClothAsset/Source/ChaosClothAssetEngine/Private/ChaosClothAsset/ClothAssetBuilder.cpp
@@ -8,7 +8,6 @@
 #include "ChaosClothAsset/ClothCollectionGroup.h"
 #include "ChaosClothAsset/CollectionClothFacade.h"
 #include "ChaosClothAsset/ClothAssetPrivate.h"
-#include "ChaosClothAsset/ClothEngineTools.h"
 #include "Chaos/CollectionPropertyFacade.h"
 #include "Engine/RendererSettings.h"
 #include "Engine/SkeletalMesh.h"
@@ -152,7 +151,18 @@ void UChaosClothAsset::FBuilder::BuildLod(FSkeletalMeshLODModel& LODModel, const
 
 	// Retrieve MaxDistance information (weight map and Low/High values)
 	FCollectionPropertyConstFacade Properties(ClothCollection);
-	const FPointWeightMap MaxDistances = FClothEngineTools::GetMaxDistanceWeightMap(ClothFacade, Properties, NumLodSimVertices);
+
+	int32 MaxDistancePropertyKeyIndex;
+	FString MaxDistanceString = TEXT("MaxDistance");
+	MaxDistanceString = Properties.GetStringValue(MaxDistanceString, MaxDistanceString, &MaxDistancePropertyKeyIndex);
+	const bool bHasMaxDistanceProperty = (MaxDistancePropertyKeyIndex != INDEX_NONE);
+	const float MaxDistanceOffset = bHasMaxDistanceProperty ? Properties.GetLowValue<float>(MaxDistancePropertyKeyIndex) : TNumericLimits<float>::Max();  // Uses infinite distance when no MaxDistance properties are set
+	const float MaxDistanceScale = bHasMaxDistanceProperty ? Properties.GetHighValue<float>(MaxDistancePropertyKeyIndex) - MaxDistanceOffset : 0.f;
+	const TConstArrayView<float> MaxDistanceWeightMap = ClothFacade.GetWeightMap(FName(MaxDistanceString));
+
+	const FPointWeightMap MaxDistances = (MaxDistanceWeightMap.Num() == NumLodSimVertices) ?
+		FPointWeightMap(MaxDistanceWeightMap, MaxDistanceOffset, MaxDistanceScale) :
+		FPointWeightMap(NumLodSimVertices, MaxDistanceOffset);
 
 	const int32 NumRenderVertices = ClothFacade.GetNumRenderVertices();
 	LODModel.MeshToImportVertexMap.Reserve(NumRenderVertices);
diff --git a/Engine/Plugins/ChaosClothAsset/Source/ChaosClothAssetEngine/Private/ChaosClothAsset/ClothAssetSKMClothingAsset.cpp b/Engine/Plugins/ChaosClothAsset/Source/ChaosClothAssetEngine/Private/ChaosClothAsset/ClothAssetSKMClothingAsset.cpp
deleted file mode 100644
diff --git a/Engine/Plugins/ChaosClothAsset/Source/ChaosClothAssetEngine/Private/ChaosClothAsset/ClothAssetSKMClothingAsset.h b/Engine/Plugins/ChaosClothAsset/Source/ChaosClothAssetEngine/Private/ChaosClothAsset/ClothAssetSKMClothingAsset.h
deleted file mode 100644
diff --git a/Engine/Plugins/ChaosClothAsset/Source/ChaosClothAssetEngine/Private/ChaosClothAsset/ClothAssetSKMClothingAssetFactory.cpp b/Engine/Plugins/ChaosClothAsset/Source/ChaosClothAssetEngine/Private/ChaosClothAsset/ClothAssetSKMClothingAssetFactory.cpp
deleted file mode 100644
diff --git a/Engine/Plugins/ChaosClothAsset/Source/ChaosClothAssetEngine/Private/ChaosClothAsset/ClothAssetSKMClothingAssetFactory.h b/Engine/Plugins/ChaosClothAsset/Source/ChaosClothAssetEngine/Private/ChaosClothAsset/ClothAssetSKMClothingAssetFactory.h
deleted file mode 100644
diff --git a/Engine/Plugins/ChaosClothAsset/Source/ChaosClothAssetEngine/Private/ChaosClothAsset/ClothAssetSKMSimulationFactory.cpp b/Engine/Plugins/ChaosClothAsset/Source/ChaosClothAssetEngine/Private/ChaosClothAsset/ClothAssetSKMSimulationFactory.cpp
deleted file mode 100644
diff --git a/Engine/Plugins/ChaosClothAsset/Source/ChaosClothAssetEngine/Private/ChaosClothAsset/ClothAssetSKMSimulationFactory.h b/Engine/Plugins/ChaosClothAsset/Source/ChaosClothAssetEngine/Private/ChaosClothAsset/ClothAssetSKMSimulationFactory.h
deleted file mode 100644
diff --git a/Engine/Plugins/ChaosClothAsset/Source/ChaosClothAssetEngine/Private/ChaosClothAsset/ClothEngineTools.cpp b/Engine/Plugins/ChaosClothAsset/Source/ChaosClothAssetEngine/Private/ChaosClothAsset/ClothEngineTools.cpp
--- a/Engine/Plugins/ChaosClothAsset/Source/ChaosClothAssetEngine/Private/ChaosClothAsset/ClothEngineTools.cpp
+++ b/Engine/Plugins/ChaosClothAsset/Source/ChaosClothAssetEngine/Private/ChaosClothAsset/ClothEngineTools.cpp
@@ -1,7 +1,6 @@
 // Copyright Epic Games, Inc. All Rights Reserved.
 
 #include "ChaosClothAsset/ClothEngineTools.h"
-#include "Chaos/CollectionPropertyFacade.h"
 #include "ChaosClothAsset/CollectionClothFacade.h"
 #include "ChaosClothAsset/CollectionClothSelectionFacade.h"
 #include "ChaosClothAsset/ClothCollectionGroup.h"
@@ -123,26 +122,4 @@ void FClothEngineTools::GenerateTethersFromCustomSelectionSets(const TSharedRef<
 	}
 }
 
-FPointWeightMap FClothEngineTools::GetMaxDistanceWeightMap(const TSharedRef<const FManagedArrayCollection>& ClothCollection, const int32 NumLodSimVertices)
-{
-	FCollectionClothConstFacade ClothFacade(ClothCollection);
-	::Chaos::Softs::FCollectionPropertyConstFacade PropertyFacade(ClothCollection);
-	return GetMaxDistanceWeightMap(ClothFacade, PropertyFacade, NumLodSimVertices);
-}
-
-FPointWeightMap FClothEngineTools::GetMaxDistanceWeightMap(const FCollectionClothConstFacade& ClothFacade, const ::Chaos::Softs::FCollectionPropertyConstFacade& PropertyFacade, const int32 NumLodSimVertices)
-{
-	int32 MaxDistancePropertyKeyIndex;
-	FString MaxDistanceString = TEXT("MaxDistance");
-	MaxDistanceString = PropertyFacade.GetStringValue(MaxDistanceString, MaxDistanceString, &MaxDistancePropertyKeyIndex);
-	const bool bHasMaxDistanceProperty = (MaxDistancePropertyKeyIndex != INDEX_NONE);
-	const float MaxDistanceOffset = bHasMaxDistanceProperty ? PropertyFacade.GetLowValue<float>(MaxDistancePropertyKeyIndex) : TNumericLimits<float>::Max();  // Uses infinite distance when no MaxDistance properties are set
-	const float MaxDistanceScale = bHasMaxDistanceProperty ? PropertyFacade.GetHighValue<float>(MaxDistancePropertyKeyIndex) - MaxDistanceOffset : 0.f;
-	const TConstArrayView<float> MaxDistanceWeightMap = ClothFacade.GetWeightMap(FName(MaxDistanceString));
-
-	return (MaxDistanceWeightMap.Num() == NumLodSimVertices) ?
-		FPointWeightMap(MaxDistanceWeightMap, MaxDistanceOffset, MaxDistanceScale) :
-		FPointWeightMap(NumLodSimVertices, MaxDistanceOffset);
-}
-
 }  // End namespace UE::Chaos::ClothAsset
diff --git a/Engine/Plugins/ChaosClothAsset/Source/ChaosClothAssetEngine/Public/ChaosClothAsset/ClothEngineTools.h b/Engine/Plugins/ChaosClothAsset/Source/ChaosClothAssetEngine/Public/ChaosClothAsset/ClothEngineTools.h
--- a/Engine/Plugins/ChaosClothAsset/Source/ChaosClothAssetEngine/Public/ChaosClothAsset/ClothEngineTools.h
+++ b/Engine/Plugins/ChaosClothAsset/Source/ChaosClothAssetEngine/Public/ChaosClothAsset/ClothEngineTools.h
@@ -2,27 +2,17 @@
 
 #pragma once
 
-#include "Math/Vector2D.h"
 #include "Templates/SharedPointer.h"
-#include "PointWeightMap.h"
+#include "Math/Vector2D.h"
 
 #define UE_API CHAOSCLOTHASSETENGINE_API
 
 struct FManagedArrayCollection;
 class FName;
 class FSkeletalMeshLODModel;
-namespace Chaos
-{
-	namespace Softs
-	{
-		class FCollectionPropertyConstFacade;
-	}
-}
 
 namespace UE::Chaos::ClothAsset
 {
-	class FCollectionClothConstFacade;
-
 	/**
 	 *  Tools operating on cloth collections with Engine dependency
 	 */
@@ -33,11 +23,6 @@ namespace UE::Chaos::ClothAsset
 		static UE_API void GenerateTethersFromSelectionSet(const TSharedRef<FManagedArrayCollection>& ClothCollection, const FName& FixedEndSet, const bool bGeodesicTethers);
 		/** @param CustomTetherEndSets: First element of each pair is DynamicSet, second is FixedSet */
 		static UE_API void GenerateTethersFromCustomSelectionSets(const TSharedRef<FManagedArrayCollection>& ClothCollection, const FName& FixedEndSet, const TArray<TPair<FName, FName>>& CustomTetherEndSets, const bool bGeodesicTethers);
-	
-		/** Retrieve the MaxDistance weight map from the cloth facades. Return an empty map if the collection doesn't match the requested NumLodSimVertices.*/
-		static FPointWeightMap GetMaxDistanceWeightMap(const FCollectionClothConstFacade& ClothFacade, const ::Chaos::Softs::FCollectionPropertyConstFacade& PropertyFacade, const int32 NumLodSimVertices);
-		/** Retrieve the MaxDistance weight map from the cloth collection. Return an empty map if the collection doesn't match the requested NumLodSimVertices. */
-		static FPointWeightMap GetMaxDistanceWeightMap(const TSharedRef<const FManagedArrayCollection>& ClothCollection, const int32 NumLodSimVertices);
 	};
 }  // End namespace UE::Chaos::ClothAsset
 
diff --git a/Engine/Source/Runtime/ClothingSystemRuntimeCommon/Private/ClothingAsset.cpp b/Engine/Source/Runtime/ClothingSystemRuntimeCommon/Private/ClothingAsset.cpp
--- a/Engine/Source/Runtime/ClothingSystemRuntimeCommon/Private/ClothingAsset.cpp
+++ b/Engine/Source/Runtime/ClothingSystemRuntimeCommon/Private/ClothingAsset.cpp
@@ -33,6 +33,7 @@
 
 #include UE_INLINE_GENERATED_CPP_BY_NAME(ClothingAsset)
 
+DEFINE_LOG_CATEGORY(LogClothingAsset)
 #define LOCTEXT_NAMESPACE "ClothingAsset"
 
 //==============================================================================
@@ -125,6 +126,16 @@ UClothingAssetCommon::UClothingAssetCommon(const FObjectInitializer& ObjectIniti
 }
 
 #if WITH_EDITOR
+
+void Warn(const FText& Error)
+{
+	FNotificationInfo* const NotificationInfo = new FNotificationInfo(Error);
+	NotificationInfo->ExpireDuration = 5.0f;
+	FSlateNotificationManager::Get().QueueNotification(NotificationInfo);
+
+	UE_LOG(LogClothingAsset, Warning, TEXT("%s"), *Error.ToString());
+}
+
 bool UClothingAssetCommon::BindToSkeletalMesh(
 	USkeletalMesh* InSkelMesh, 
 	const int32 InMeshLodIndex, 
@@ -147,43 +158,43 @@ bool UClothingAssetCommon::BindToSkeletalMesh(
 	// If we've been added to the wrong mesh
 	if(InSkelMesh != GetOuter())
 	{
-		const FText Error = FText::Format(
+		FText Error = FText::Format(
 			LOCTEXT("Error_WrongMesh", "Failed to bind clothing asset {0} as the provided mesh is not the owner of this asset."), 
 			FText::FromString(GetName()));
-		WarningNotification(Error);
+		Warn(Error);
 		return false;
 	}
 
 	// If we don't have clothing data
 	if(!LodData.IsValidIndex(InAssetLodIndex))
 	{
-		const FText Error = FText::Format(
+		FText Error = FText::Format(
 			LOCTEXT("Error_NoClothingLod", "Failed to bind clothing asset {0} LOD{1} as LOD{2} does not exist."), 
 			FText::FromString(GetName()), 
 			InAssetLodIndex,
 			InAssetLodIndex);
-		WarningNotification(Error);
+		Warn(Error);
 		return false;
 	}
 
 	// If we don't have a mesh
 	if(!InSkelMesh)
 	{
-		const FText Error = FText::Format(
+		FText Error = FText::Format(
 			LOCTEXT("Error_NoMesh", "Failed to bind clothing asset {0} as provided skel mesh does not exist."), 
 			FText::FromString(GetName()));
-		WarningNotification(Error);
+		Warn(Error);
 		return false;
 	}
 
 	// If the mesh LOD index is invalid
 	if(!InSkelMesh->GetImportedModel()->LODModels.IsValidIndex(InMeshLodIndex))
 	{
-		const FText Error = FText::Format(
+		FText Error = FText::Format(
 			LOCTEXT("Error_InvalidMeshLOD", "Failed to bind clothing asset {0} as mesh LOD{1} does not exist."), 
 			FText::FromString(GetName()), 
 			InMeshLodIndex);
-		WarningNotification(Error);
+		Warn(Error);
 		return false;
 	}
 
@@ -193,13 +204,13 @@ bool UClothingAssetCommon::BindToSkeletalMesh(
 		const int32& MappedLod = LodMap[MapIndex];
 		if(MappedLod == InAssetLodIndex)
 		{
-			const FText Error = FText::Format(
+			FText Error = FText::Format(
 				LOCTEXT("Error_LodMapped", "Failed to bind clothing asset {0} LOD{1} as LOD{2} is already mapped to mesh LOD{3}."), 
 				FText::FromString(GetName()), 
 				InAssetLodIndex, 
 				InAssetLodIndex, 
 				MapIndex);
-			WarningNotification(Error);
+			Warn(Error);
 			return false;
 		}
 	}
@@ -273,13 +284,13 @@ bool UClothingAssetCommon::BindToSkeletalMesh(
 	if(TempBoneMap.Num() > FGPUBaseSkinVertexFactory::GetMaxGPUSkinBones())
 	{
 		// Failed to apply as we've exceeded the number of bones we can skin
-		const FText Error = FText::Format(
+		FText Error = FText::Format(
 			LOCTEXT("Error_TooManyBones", "Failed to bind clothing asset {0} LOD{1} as this causes the section to require {2} bones. The maximum per section is currently {3}."), 
 			FText::FromString(GetName()), 
 			InAssetLodIndex, 
 			TempBoneMap.Num(), 
 			FGPUBaseSkinVertexFactory::GetMaxGPUSkinBones());
-		WarningNotification(Error);
+		Warn(Error);
 		return false;
 	}
 
@@ -526,7 +537,7 @@ void UClothingAssetCommon::UpdateLODBiasMappings(const USkeletalMesh* SkeletalMe
 					SectionIndex,
 					UpdatedLODIndex, 
 					LODIndex);
-				WarningNotification(Error);
+				Warn(Error);
 				continue;
 			}
 
@@ -614,11 +625,11 @@ void UClothingAssetCommon::UnbindFromSkeletalMesh(
 	{
 		if(!Mesh->LODModels.IsValidIndex(InMeshLodIndex))
 		{
-			const FText Error = FText::Format(
+			FText Error = FText::Format(
 				LOCTEXT("Error_UnbindNoMeshLod", "Failed to remove clothing asset {0} from mesh LOD{1} as that LOD doesn't exist."), 
 				FText::FromString(GetName()), 
 				InMeshLodIndex);
-			WarningNotification(Error);
+			Warn(Error);
 
 			return;
 		}
diff --git a/Engine/Source/Runtime/ClothingSystemRuntimeCommon/Public/ClothingAsset.h b/Engine/Source/Runtime/ClothingSystemRuntimeCommon/Public/ClothingAsset.h
--- a/Engine/Source/Runtime/ClothingSystemRuntimeCommon/Public/ClothingAsset.h
+++ b/Engine/Source/Runtime/ClothingSystemRuntimeCommon/Public/ClothingAsset.h
@@ -10,6 +10,8 @@
 #include "Algo/AnyOf.h"
 #include "ClothingAsset.generated.h"
 
+DECLARE_LOG_CATEGORY_EXTERN(LogClothingAsset, Log, All);
+
 class UPhysicsAsset;
 struct FSkelMeshSection;
 
diff --git a/Engine/Source/Runtime/ClothingSystemRuntimeInterface/ClothingSystemRuntimeInterface.Build.cs b/Engine/Source/Runtime/ClothingSystemRuntimeInterface/ClothingSystemRuntimeInterface.Build.cs
--- a/Engine/Source/Runtime/ClothingSystemRuntimeInterface/ClothingSystemRuntimeInterface.Build.cs
+++ b/Engine/Source/Runtime/ClothingSystemRuntimeInterface/ClothingSystemRuntimeInterface.Build.cs
@@ -16,10 +16,5 @@ public class ClothingSystemRuntimeInterface : ModuleRules
                 "CoreUObject"
             }
         );
-
-		if (Target.bBuildEditor || Target.bCompileAgainstEditor)
-		{
-			PrivateDependencyModuleNames.Add("Slate");  // For editor notifications
-		}
-	}
+    }
 }
diff --git a/Engine/Source/Runtime/ClothingSystemRuntimeInterface/Private/ClothingAssetBase.cpp b/Engine/Source/Runtime/ClothingSystemRuntimeInterface/Private/ClothingAssetBase.cpp
deleted file mode 100644
diff --git a/Engine/Source/Runtime/ClothingSystemRuntimeInterface/Public/ClothingAssetBase.h b/Engine/Source/Runtime/ClothingSystemRuntimeInterface/Public/ClothingAssetBase.h
--- a/Engine/Source/Runtime/ClothingSystemRuntimeInterface/Public/ClothingAssetBase.h
+++ b/Engine/Source/Runtime/ClothingSystemRuntimeInterface/Public/ClothingAssetBase.h
@@ -8,10 +8,6 @@
 
 class USkeletalMesh;
 
-#define UE_API CLOTHINGSYSTEMRUNTIMEINTERFACE_API
-
-UE_API DECLARE_LOG_CATEGORY_EXTERN(LogClothingAsset, Log, All);
-
 /**
  * An interface object for any clothing asset the engine can use.
  * Any clothing asset concrete object should derive from this.
@@ -58,9 +54,7 @@ public:
 	/** Add a new LOD class instance. */
 	UE_DEPRECATED(5.7, "Use UClothingAssetCommon::AddNewLod() instead.")
 	virtual int32 AddNewLod()
-	{
-		return INDEX_NONE;
-	}
+	PURE_VIRTUAL(UClothingAssetBase::AddNewLod(), return INDEX_NONE;);
 
 	/**
 	*	Builds the LOD transition data
@@ -68,7 +62,8 @@ public:
 	*	in exactly the same way the render mesh is skinned to create a smooth swap
 	*/
 	UE_DEPRECATED(5.7, "Use UClothingAssetCommon::BuildLodTransitionData() instead.")
-	virtual void BuildLodTransitionData() {}
+	virtual void BuildLodTransitionData()
+	PURE_VIRTUAL(UClothingAssetBase::BuildLodTransitionData(), );
 #endif
 
 #if WITH_EDITORONLY_DATA
@@ -77,7 +72,8 @@ public:
 	 * has changed, so any intermediate generated data can be regenerated.
 	 */
 	UE_DEPRECATED(5.7, "Use UClothingAssetCommon::InvalidateAllCachedData() instead.")
-	virtual void InvalidateAllCachedData() {}
+	virtual void InvalidateAllCachedData()
+	PURE_VIRTUAL(UClothingAssetBase::InvalidateAllCachedData(),);
 #endif // WITH_EDITORONLY_DATA
 
 	/** 
@@ -92,16 +88,12 @@ public:
 	/** Check the validity of a LOD index */
 	UE_DEPRECATED(5.7, "Use UClothingAssetCommon::IsValidLod() instead.")
 	virtual bool IsValidLod(int32 InLodIndex) const
-	{
-		return false;
-	}
+	PURE_VIRTUAL(UClothingAssetBase::IsValidLod(), return false;);
 
 	/** Get the number of LODs defined in the clothing asset */
 	UE_DEPRECATED(5.7, "Use UClothingAssetCommon::GetNumLods() instead.")
 	virtual int32 GetNumLods() const
-	{
-		return 0;
-	}
+	PURE_VIRTUAL(UClothingAssetBase::GetNumLods(), return 0;);
 
 	/** Called after all cloth assets sharing the same simulation are added or loaded */
 	virtual void PostUpdateAllAssets()
@@ -124,14 +116,7 @@ protected:
 	/** The asset factory should have access, as it will assign the asset guid when building assets */
 	friend class UClothingAssetFactory;
 
-#if WITH_EDITOR
-	/** Warn using slate's notification manager. Use this method to notify binding issues back to the user. */
-	static UE_API void WarningNotification(const FText& Text);
-#endif
-
 	/** Guid to identify this asset. Will be embedded into chunks that are created using this asset */
 	UPROPERTY()
 	FGuid AssetGuid;
 };
-
-#undef UE_API
