Cloth - Makes the skeletal mesh binding utility functions work with all types of skeletal mesh clothing data.

#rnx

diff --git a/Engine/Plugins/Interchange/Runtime/Source/Import/Private/Mesh/InterchangeSkeletalMeshFactory.cpp b/Engine/Plugins/Interchange/Runtime/Source/Import/Private/Mesh/InterchangeSkeletalMeshFactory.cpp
--- a/Engine/Plugins/Interchange/Runtime/Source/Import/Private/Mesh/InterchangeSkeletalMeshFactory.cpp
+++ b/Engine/Plugins/Interchange/Runtime/Source/Import/Private/Mesh/InterchangeSkeletalMeshFactory.cpp
@@ -2383,9 +2383,9 @@ void UInterchangeSkeletalMeshFactory::FinalizeObject_GameThread(const FSetupObje
 		FSkinnedAssetAsyncBuildScope AsyncBuildScope(SkeletalMesh);
 		for (ClothingAssetUtils::FClothingAssetMeshBinding& ExistingClothMeshBinding : ImportAssetObjectData.ExistingClothingBindings)
 		{
-			if (UClothingAssetCommon* ClothAssetCommon = ExistingClothMeshBinding.Asset)
+			if (ExistingClothMeshBinding.Asset)
 			{
-				ClothAssetCommon->RefreshBoneMapping(SkeletalMesh);
+				ExistingClothMeshBinding.Asset->RefreshBoneMapping(SkeletalMesh);
 			}
 		}
 		FSkeletalMeshModel* ImportedResource = SkeletalMesh->GetImportedModel();
diff --git a/Engine/Source/Runtime/ClothingSystemRuntimeCommon/Private/ClothingAsset.cpp b/Engine/Source/Runtime/ClothingSystemRuntimeCommon/Private/ClothingAsset.cpp
--- a/Engine/Source/Runtime/ClothingSystemRuntimeCommon/Private/ClothingAsset.cpp
+++ b/Engine/Source/Runtime/ClothingSystemRuntimeCommon/Private/ClothingAsset.cpp
@@ -86,10 +86,9 @@ void ClothingAssetUtils::GetAllLodMeshClothingAssetBindings(const USkeletalMesh*
 			const FSkelMeshSection& Section = MeshLODModel.Sections[SectionIndex];
 			if (Section.HasClothingData())
 			{
-				UClothingAssetBase* ClothingAsset = SkeletalMesh->GetClothingAsset(Section.ClothingData.AssetGuid);
 				FClothingAssetMeshBinding ClothBinding;
-				ClothBinding.Asset = Cast<UClothingAssetCommon>(ClothingAsset);
-				ClothBinding.AssetInternalLodIndex = Section.ClothingData.AssetLodIndex;// InSkelMesh->GetClothingAssetIndex(Section.ClothingData.AssetGuid);
+				ClothBinding.Asset = SkeletalMesh->GetClothingAsset(Section.ClothingData.AssetGuid);
+				ClothBinding.AssetInternalLodIndex = Section.ClothingData.AssetLodIndex;
 				check(ClothBinding.AssetInternalLodIndex == Section.ClothingData.AssetLodIndex);
 				ClothBinding.LODIndex = InLodIndex;
 				ClothBinding.SectionIndex = SectionIndex;
diff --git a/Engine/Source/Runtime/ClothingSystemRuntimeCommon/Public/ClothingAsset.h b/Engine/Source/Runtime/ClothingSystemRuntimeCommon/Public/ClothingAsset.h
--- a/Engine/Source/Runtime/ClothingSystemRuntimeCommon/Public/ClothingAsset.h
+++ b/Engine/Source/Runtime/ClothingSystemRuntimeCommon/Public/ClothingAsset.h
@@ -22,7 +22,7 @@ namespace ClothingAssetUtils
 	 */
 	struct FClothingAssetMeshBinding
 	{
-		UClothingAssetCommon* Asset;
+		UClothingAssetBase* Asset;
 		int32 LODIndex;
 		int32 SectionIndex;
 		int32 AssetInternalLodIndex;
