Skeletal Mesh Cloth - Added a new action to remove the last Clothing LOD.

[RN] Skeletal Mesh Cloth - Added a new action to remove the last Clothing LOD.
#tests Removed a bunch of LODs, tested with undo/redo.

diff --git a/Engine/Source/Editor/ClothPainter/Private/ClothingAssetListCommands.cpp b/Engine/Source/Editor/ClothPainter/Private/ClothingAssetListCommands.cpp
--- a/Engine/Source/Editor/ClothPainter/Private/ClothingAssetListCommands.cpp
+++ b/Engine/Source/Editor/ClothPainter/Private/ClothingAssetListCommands.cpp
@@ -16,7 +16,8 @@
 
 void FClothingAssetListCommands::RegisterCommands()
 {
-	UI_COMMAND(DeleteAsset, "Delete Asset", "Deletes a clothing asset from the mesh.", EUserInterfaceActionType::Button, FInputChord());
+	UI_COMMAND(DeleteAsset, "Delete Asset", "Delete a Clothing Data from the mesh.", EUserInterfaceActionType::Button, FInputChord());
+	UI_COMMAND(DeleteAssetLastLod, "Delete Asset's Last LOD", "Delete the last clothing Data LOD.", EUserInterfaceActionType::Button, FInputChord());
 	UI_COMMAND(RebuildAssetParams, "Rebuild Asset Parameter Masks", "Takes the parameter masks in LOD0 and creates masks in all lower LODs to match, casting those parameter masks to the new mesh.", EUserInterfaceActionType::Button, FInputChord());
 
 	// Add one command per exporter provider
diff --git a/Engine/Source/Editor/ClothPainter/Private/ClothingAssetListCommands.h b/Engine/Source/Editor/ClothPainter/Private/ClothingAssetListCommands.h
--- a/Engine/Source/Editor/ClothPainter/Private/ClothingAssetListCommands.h
+++ b/Engine/Source/Editor/ClothPainter/Private/ClothingAssetListCommands.h
@@ -27,7 +27,7 @@ public:
 	virtual void RegisterCommands() override;
 
 	TSharedPtr<FUICommandInfo> DeleteAsset;
-
+	TSharedPtr<FUICommandInfo> DeleteAssetLastLod;
 	TSharedPtr<FUICommandInfo> RebuildAssetParams;
 	TMap<FName, TSharedPtr<FUICommandInfo>> ExportAssets;
 };
diff --git a/Engine/Source/Editor/ClothPainter/Private/SClothAssetSelector.cpp b/Engine/Source/Editor/ClothPainter/Private/SClothAssetSelector.cpp
--- a/Engine/Source/Editor/ClothPainter/Private/SClothAssetSelector.cpp
+++ b/Engine/Source/Editor/ClothPainter/Private/SClothAssetSelector.cpp
@@ -199,6 +199,12 @@ public:
 			FExecuteAction::CreateSP(this, &SAssetListRow::DeleteAsset)
 		);
 
+		UICommandList->MapAction(
+			Commands.DeleteAssetLastLod,
+			FExecuteAction::CreateSP(this, &SAssetListRow::DeleteAssetLastLod),
+			FCanExecuteAction::CreateSP(this, &SAssetListRow::CanDeleteAssetLastLod)
+		);
+
 		UICommandList->MapAction(
 			Commands.RebuildAssetParams,
 			FExecuteAction::CreateSP(this, &SAssetListRow::RebuildLODParameters),
@@ -221,12 +227,20 @@ public:
 	{
 		if(Item.IsValid() && MouseEvent.GetEffectingButton() == EKeys::RightMouseButton)
 		{
+			const UClothingAssetCommon* const Asset = Cast<UClothingAssetCommon>(Item->ClothingAsset.Get());  // Only the UClothingAssetCommon type have LODs
+			const int32 LastLodIndex = Asset ? Asset->GetNumLods() - 1 : INDEX_NONE;
+
 			const FClothingAssetListCommands& Commands = FClothingAssetListCommands::Get();
 			FMenuBuilder Builder(true, UICommandList);
 
 			Builder.BeginSection(NAME_None, LOCTEXT("AssetActions_SectionName", "Actions"));
 			{
 				Builder.AddMenuEntry(FGenericCommands::Get().Delete);
+				if (LastLodIndex > 0)  // Can't delete LOD0
+				{
+					Builder.AddMenuEntry(Commands.DeleteAssetLastLod, NAME_None,
+						FText::Format(LOCTEXT("DeleteAssetLastLod", "Delete LOD{0}"), LastLodIndex));
+				}
 				Builder.AddMenuEntry(Commands.RebuildAssetParams);
 
 				for (const TPair<FName, TSharedPtr<FUICommandInfo>>& CommandId : Commands.ExportAssets)
@@ -303,6 +317,42 @@ private:
 		}
 	}
 
+	bool CanDeleteAssetLastLod()
+	{
+		if (Item.IsValid())
+		{
+			if (UClothingAssetCommon* const Asset = Cast<UClothingAssetCommon>(Item->ClothingAsset.Get()))  // Only the UClothingAssetCommon type have LODs
+			{
+				return Asset->GetNumLods() > 1;  // Can't delete LOD0
+			}
+		}
+		return false;
+	}
+
+	void DeleteAssetLastLod()
+	{
+		if (Item.IsValid())
+		{
+			if (UClothingAssetCommon* const Asset = Cast<UClothingAssetCommon>(Item->ClothingAsset.Get()))  // Only the UClothingAssetCommon type have LODs
+			{
+				if (ensure(Asset->GetNumLods() > 1))  // Can't delete LOD0
+				{
+					if (USkeletalMesh* const SkeletalMesh = Cast<USkeletalMesh>(Asset->GetOuter()))
+					{
+						FScopedSuspendAlternateSkinWeightPreview ScopedSuspendAlternateSkinnWeightPreview(SkeletalMesh);
+						FScopedSkeletalMeshPostEditChange ScopedPostEditChange(SkeletalMesh);
+						FScopedTransaction CurrTransaction(LOCTEXT("DeleteLastLod_Transaction", "Delete the last Clothing Data LOD."));
+						SkeletalMesh->Modify();
+						Asset->Modify();
+						SkeletalMesh->PreEditChange(nullptr);
+						Asset->RemoveLastLod();
+						OnInvalidateList.ExecuteIfBound();
+					}
+				}
+			}
+		}
+	}
+
 	// Using LOD0 of an asset, rebuild the other LOD masks by mapping the LOD0 parameters onto their meshes
 	void RebuildLODParameters()
 	{
@@ -1512,6 +1562,17 @@ void SClothAssetSelector::RefreshAssetList()
 			AssetList->SetSelection(AssetListItems[SelectedItem]);
 		}
 	}
+
+	// Refresh LOD list
+	constexpr bool bRefreshMasks = false;  // Will be done in RefreshMaskList() running afterwards
+	if (const UClothingAssetCommon* const ClothingAssetCommon = Cast<UClothingAssetCommon>(CurrSelectedAsset))  // Only the UClothingAssetCommon type assets have LODs
+	{
+		SetSelectedLod(FMath::Min(SelectedLod, ClothingAssetCommon->GetNumLods() - 1), bRefreshMasks);
+	}
+	else
+	{
+		SetSelectedLod(INDEX_NONE, bRefreshMasks);
+	}
 }
 
 void SClothAssetSelector::RefreshMaskList()
diff --git a/Engine/Source/Runtime/ClothingSystemRuntimeCommon/Private/ClothingAsset.cpp b/Engine/Source/Runtime/ClothingSystemRuntimeCommon/Private/ClothingAsset.cpp
--- a/Engine/Source/Runtime/ClothingSystemRuntimeCommon/Private/ClothingAsset.cpp
+++ b/Engine/Source/Runtime/ClothingSystemRuntimeCommon/Private/ClothingAsset.cpp
@@ -789,31 +789,33 @@ void UClothingAssetCommon::ApplyParameterMasks(bool bUpdateFixedVertData, bool b
 	
 	if (bUpdateFixedVertData && SkeletalMesh)
 	{
-		FSkeletalMeshModel* Resource = SkeletalMesh->GetImportedModel();
-		FScopedSkeletalMeshPostEditChange ScopedSkeletalMeshPostEditChange(SkeletalMesh);
+		if (FSkeletalMeshModel* const Resource = SkeletalMesh->GetImportedModel())
+		{
+			FScopedSkeletalMeshPostEditChange ScopedSkeletalMeshPostEditChange(SkeletalMesh);
 
-		SkeletalMesh->PreEditChange(nullptr);
+			SkeletalMesh->PreEditChange(nullptr);
 
-		for (FSkeletalMeshLODModel& LodModel : Resource->LODModels)
-		{
-			for (FSkelMeshSection& Section : LodModel.Sections)
+			for (FSkeletalMeshLODModel& LodModel : Resource->LODModels)
 			{
-				if (!Section.HasClothingData() || Cast<UClothingAssetCommon>(SkeletalMesh->GetClothingAsset(Section.ClothingData.AssetGuid)) != this)
+				for (FSkelMeshSection& Section : LodModel.Sections)
 				{
-					continue;
-				}
-				const FClothLODDataCommon& LodDatum = LodData[Section.ClothingData.AssetLodIndex];
-				const FPointWeightMap* const MaxDistances = LodDatum.PhysicalMeshData.FindWeightMap(EWeightMapTargetCommon::MaxDistance);
+					if (!Section.HasClothingData() || Cast<UClothingAssetCommon>(SkeletalMesh->GetClothingAsset(Section.ClothingData.AssetGuid)) != this)
+					{
+						continue;
+					}
+					const FClothLODDataCommon& LodDatum = LodData[Section.ClothingData.AssetLodIndex];
+					const FPointWeightMap* const MaxDistances = LodDatum.PhysicalMeshData.FindWeightMap(EWeightMapTargetCommon::MaxDistance);
 
-				for (TArray<FMeshToMeshVertData>& ClothMappingData : Section.ClothMappingDataLODs)
-				{
-					ClothingMeshUtils::ComputeVertexContributions(ClothMappingData, MaxDistances, LodDatum.bSmoothTransition, LodDatum.bUseMultipleInfluences);
+					for (TArray<FMeshToMeshVertData>& ClothMappingData : Section.ClothMappingDataLODs)
+					{
+						ClothingMeshUtils::ComputeVertexContributions(ClothMappingData, MaxDistances, LodDatum.bSmoothTransition, LodDatum.bUseMultipleInfluences);
+					}
 				}
 			}
-		}
-		if (bInvalidateDerivedDataCache)  // We must always dirty the DDC key unless previewing
-		{
-			SkeletalMesh->InvalidateDeriveDataCacheGUID();
+			if (bInvalidateDerivedDataCache)  // We must always dirty the DDC key unless previewing
+			{
+				SkeletalMesh->InvalidateDeriveDataCacheGUID();
+			}
 		}
 	}
 }
@@ -1556,6 +1558,33 @@ int32 UClothingAssetCommon::AddNewLod()
 	return LodData.AddDefaulted();
 }
 
+void UClothingAssetCommon::RemoveLastLod()
+{
+	const int32 LodIndex = LodData.Num() - 1;
+	if (LodIndex > 0)  // Can't remove LOD0
+	{
+		TOptional<FScopedSkeletalMeshPostEditChange> SkeletalMeshPostEditChange;
+		USkeletalMesh* const SkeletalMesh = Cast<USkeletalMesh>(GetOuter());
+		if (SkeletalMesh)
+		{
+			SkeletalMeshPostEditChange.Emplace(SkeletalMesh);
+		}
+
+		int32 MeshLodIndex;
+		while (LodMap.Find(LodIndex, MeshLodIndex))
+		{
+			LodMap[MeshLodIndex] = INDEX_NONE;  // Note: This must be done prior calling UnbindFromSkeletalMesh(), so that the info is correct when rebuilding the mesh
+
+			if (SkeletalMesh)
+			{
+				// Unbind from the Skeletal Mesh
+				ClothingAssetUtils::UnbindFromSkeletalMesh(*this, *SkeletalMesh, MeshLodIndex);
+			}
+		}
+		LodData.RemoveAt(LodIndex, EAllowShrinking::Yes);
+	}
+}
+
 void UClothingAssetCommon::PostEditChangeChainProperty(FPropertyChangedChainEvent& ChainEvent)
 {
 	Super::PostEditChangeChainProperty(ChainEvent);
diff --git a/Engine/Source/Runtime/ClothingSystemRuntimeCommon/Public/ClothingAsset.h b/Engine/Source/Runtime/ClothingSystemRuntimeCommon/Public/ClothingAsset.h
--- a/Engine/Source/Runtime/ClothingSystemRuntimeCommon/Public/ClothingAsset.h
+++ b/Engine/Source/Runtime/ClothingSystemRuntimeCommon/Public/ClothingAsset.h
@@ -210,9 +210,12 @@ public:
 #endif // WITH_EDITORONLY_DATA
 
 #if WITH_EDITOR
-	/** * Add a new LOD class instance. */
+	/** Add a new LOD class instance. */
 	UE_API virtual int32 AddNewLod() override;
 
+	/** Remove the last clothing LOD (and unbind it). */
+	UE_API void RemoveLastLod();
+
 	/* Called after changes in any of the asset properties. */
 	UE_API virtual void PostEditChangeChainProperty(FPropertyChangedChainEvent& ChainEvent) override;
 #endif // WITH_EDITOR
