Chaos SKM Cloth Asset - Added missing LODBias mapping calculations.

#rnx
#Synced-CL 48812260

diff --git a/Engine/Plugins/ChaosClothAsset/Source/ChaosClothAssetEngine/Private/ChaosClothAsset/ClothAssetSKMClothingAsset.cpp b/Engine/Plugins/ChaosClothAsset/Source/ChaosClothAssetEngine/Private/ChaosClothAsset/ClothAssetSKMClothingAsset.cpp
--- a/Engine/Plugins/ChaosClothAsset/Source/ChaosClothAssetEngine/Private/ChaosClothAsset/ClothAssetSKMClothingAsset.cpp
+++ b/Engine/Plugins/ChaosClothAsset/Source/ChaosClothAssetEngine/Private/ChaosClothAsset/ClothAssetSKMClothingAsset.cpp
@@ -478,8 +478,7 @@ bool UChaosClothAssetSKMClothingAsset::BindToSkeletalMesh(USkeletalMesh* Skeleta
 	}
 
 	// Update the extra cloth deformer mapping LOD bias using this cloth entry
-	// TODO: Implement a generic form of LOD Bias mapping in ClothingAssetUtils (need to rework the clothing asset interface for this)
-	//UpdateLODBiasMappings(SkeletalMesh, MeshLodIndex, SectionIndex);
+	ClothingAssetUtils::UpdateLODBiasMappings(*this, SkeletalMesh, MeshLodIndex, SectionIndex);
 
 	return true;
 }
@@ -494,7 +493,63 @@ void UChaosClothAssetSKMClothingAsset::UnbindFromSkeletalMesh(USkeletalMesh* Ske
 
 void UChaosClothAssetSKMClothingAsset::UpdateAllLODBiasMappings(USkeletalMesh* SkeletalMesh)
 {
-	// TODO: Implement a generic form of LOD Bias mapping in ClothingAssetUtils (need to rework the clothing asset interface for this)
+	check(SkeletalMesh);
+
+	if (FSkeletalMeshModel* const SkeletalMeshModel = SkeletalMesh->GetImportedModel())
+	{
+		// Iterate through all source LODs with cloth that could lead to upper (raytraced) sections needing some additional mapping
+		TIndirectArray<FSkeletalMeshLODModel>& LODModels = SkeletalMeshModel->LODModels;
+
+		for (int32 LODIndex = LODModels.Num() - 1; LODIndex >= 0; --LODIndex)  // Iterate in reverse order to allow shrinking the mapping array first
+		{
+			// Go through all sections to find the one(s?) that uses this cloth asset and clear the existing bias mappings
+			TArray<FSkelMeshSection>& Sections = LODModels[LODIndex].Sections;
+
+			for (int32 SectionIndex = 0; SectionIndex < Sections.Num(); ++SectionIndex)
+			{
+				if (Sections[SectionIndex].ClothingData.AssetGuid == GetAssetGuid())
+				{
+					Sections[SectionIndex].ClothMappingDataLODs.SetNum(1);  // Only keep ClothMappingDataLODs[0] that is the same LOD mapping to remove LOD bias mappings
+
+					if (SkeletalMesh->GetSupportRayTracing())  // This conditional is inside the loop so that the SetNum(1) can removed the unused bias mappings
+					{
+						// Updates the upper LODs mappings of the specified section from this LODIndex
+						ClothingAssetUtils::UpdateLODBiasMappings(*this, SkeletalMesh, LODIndex, SectionIndex);
+					}
+				}
+			}
+		}
+	}
+}
+
+void UChaosClothAssetSKMClothingAsset::GetSimulationMesh(const int32 InMeshLodIndex, TArray<FVector3f>& OutPositions, TArray<uint32>& OutIndices, TArray<float>& OutMaxDistances) const
+{
+	using namespace UE::Chaos::ClothAsset;
+
+	OutPositions.Reset();
+	OutIndices.Reset();
+	OutMaxDistances.Reset();
+
+	if (Asset && ClothSimulationModelIndex >= 0 && ClothSimulationModelIndex < Asset->GetNumClothSimulationModels())
+	{
+		const TSharedPtr<const FChaosClothSimulationModel> ClothSimulationModel = Asset->GetClothSimulationModel(ClothSimulationModelIndex);
+		if (ClothSimulationModel && InMeshLodIndex >= 0 && InMeshLodIndex < ClothSimulationModel->GetNumLods())
+		{
+			OutPositions = ClothSimulationModel->GetPositions(InMeshLodIndex);
+			OutIndices = ClothSimulationModel->GetIndices(InMeshLodIndex);
+
+			const TArray<TSharedRef<const FManagedArrayCollection>>& Collections = Asset->GetCollections(ClothSimulationModelIndex);
+			if (Collections.IsValidIndex(InMeshLodIndex))
+			{
+				const TSharedRef<const FManagedArrayCollection>& Collection = Collections[InMeshLodIndex];
+				FCollectionClothConstFacade ClothFacade(Collection);
+				::Chaos::Softs::FCollectionPropertyConstFacade PropertyFacade(Collection);
+
+				FPointWeightMap MaxDistances = FClothEngineTools::GetMaxDistanceWeightMap(ClothFacade, PropertyFacade, ClothSimulationModel->GetNumVertices(InMeshLodIndex));
+				OutMaxDistances = MoveTemp(MaxDistances.Values);
+			}
+		}
+	}
 }
 
 TArray<TPair<FText, FText>> UChaosClothAssetSKMClothingAsset::GetStats() const
diff --git a/Engine/Plugins/ChaosClothAsset/Source/ChaosClothAssetEngine/Public/ChaosClothAsset/ClothAssetSKMClothingAsset.h b/Engine/Plugins/ChaosClothAsset/Source/ChaosClothAssetEngine/Public/ChaosClothAsset/ClothAssetSKMClothingAsset.h
--- a/Engine/Plugins/ChaosClothAsset/Source/ChaosClothAssetEngine/Public/ChaosClothAsset/ClothAssetSKMClothingAsset.h
+++ b/Engine/Plugins/ChaosClothAsset/Source/ChaosClothAssetEngine/Public/ChaosClothAsset/ClothAssetSKMClothingAsset.h
@@ -59,6 +59,7 @@ private:
 	virtual bool BindToSkeletalMesh(USkeletalMesh* SkeletalMesh, const int32 MeshLodIndex, const int32 SectionIndex, const int32 AssetLodIndex) override;
 	virtual void UnbindFromSkeletalMesh(USkeletalMesh* SkeletalMesh, const int32 MeshLodIndexx, const int32 SectionIndex) override;
 	virtual void UpdateAllLODBiasMappings(USkeletalMesh* SkeletalMesh) override;
+	virtual void GetSimulationMesh(const int32 InMeshLodIndex, TArray<FVector3f>& OutPositions, TArray<uint32>& OutIndices, TArray<float>& OutMaxDistances) const override;
 	virtual TArray<TPair<FText, FText>> GetStats() const override;
 #endif  // #if WITH_EDITOR
 	virtual void RefreshBoneMapping(USkeletalMesh* SkeletalMesh) override {}
diff --git a/Engine/Source/Runtime/ClothingSystemRuntimeCommon/Private/ClothingAsset.cpp b/Engine/Source/Runtime/ClothingSystemRuntimeCommon/Private/ClothingAsset.cpp
--- a/Engine/Source/Runtime/ClothingSystemRuntimeCommon/Private/ClothingAsset.cpp
+++ b/Engine/Source/Runtime/ClothingSystemRuntimeCommon/Private/ClothingAsset.cpp
@@ -240,6 +240,225 @@
 	}
 }
 
+void ClothingAssetUtils::UpdateLODBiasMappings(
+	const UClothingAssetBase& Asset,
+	const USkeletalMesh* SkeletalMesh,
+	const int32 UpdatedLODIndex,
+	const int32 SectionIndex,
+	const bool bSmoothTransition,
+	const bool bUseMultipleInfluences,
+	const float SkinningKernelRadius)
+{
+	check(SkeletalMesh);
+	check(UpdatedLODIndex >= 0);
+
+	FSkeletalMeshModel* const SkeletalMeshModel = SkeletalMesh->GetImportedModel();
+	check(SkeletalMeshModel);
+
+	const bool bAddMappingsToAnyLOD = (SkeletalMesh->GetClothLODBiasMode() == EClothLODBiasMode::MappingsToAnyLOD);
+	const bool bAddMappingsToMinLOD = (SkeletalMesh->GetClothLODBiasMode() == EClothLODBiasMode::MappingsToMinLOD && SkeletalMesh->GetRayTracingMinLOD() > 0);
+
+	TIndirectArray<FSkeletalMeshLODModel>& LODModels = SkeletalMeshModel->LODModels;
+
+	// Render mesh desc helper lambda
+	TArray<FVector3f> RenderPositions;
+	TArray<FVector3f> RenderNormals;
+	TArray<FVector3f> RenderTangents;
+	TArray<uint32> RenderIndices;
+
+	auto PrepareDeformerTargetDescriptor = [&LODModels, &RenderPositions, &RenderNormals, &RenderTangents, &RenderIndices](
+			int32 LODIndex,
+			int32 SectionIndex) -> ClothingMeshUtils::ClothMeshDesc
+		{
+			FSkeletalMeshLODModel& LODModel = LODModels[LODIndex];
+			FSkelMeshSection& Section = LODModel.Sections[SectionIndex];
+
+			RenderPositions.Reset(Section.SoftVertices.Num());
+			RenderNormals.Reset(Section.SoftVertices.Num());
+			RenderTangents.Reset(Section.SoftVertices.Num());
+
+			for (const FSoftSkinVertex& SoftSkinVertex : Section.SoftVertices)
+			{
+				RenderPositions.Add(SoftSkinVertex.Position);
+				RenderNormals.Add(SoftSkinVertex.TangentZ);
+				RenderTangents.Add(SoftSkinVertex.TangentX);
+			}
+
+			const TConstArrayView<uint32> Indices = TConstArrayView<uint32>(LODModel.IndexBuffer).Slice(Section.BaseIndex, Section.NumTriangles * 3);
+
+			RenderIndices.Reset(Section.NumTriangles * 3);
+			for (uint32 VertexIndex : Indices)
+			{
+				RenderIndices.Add(VertexIndex - Section.BaseVertexIndex);
+			}
+
+			return ClothingMeshUtils::ClothMeshDesc(RenderPositions, RenderNormals, RenderTangents, RenderIndices);
+		};
+
+	// Cloth sim mesh and max distance helper lambda
+	TArray<FVector3f> SimPositions;
+	TArray<uint32> SimIndices;
+	FPointWeightMap SimMaxDistances;
+
+	auto PrepareDeformerSourceDescriptor = [&SimPositions, &SimIndices, &SimMaxDistances](
+			const UClothingAssetBase& ClothingAsset,
+			const int32 LODIndex,
+			FPointWeightMap*& OutMaxDistances) -> ClothingMeshUtils::ClothMeshDesc
+		{
+			TArray<float> TempMaxDistances;
+			ClothingAsset.GetSimulationMesh(LODIndex, SimPositions, SimIndices, TempMaxDistances);
+			if (TempMaxDistances.Num() == SimPositions.Num())
+			{
+				SimMaxDistances.Values = MoveTemp(TempMaxDistances);
+				OutMaxDistances = &SimMaxDistances;
+			}
+			else
+			{
+				OutMaxDistances = nullptr;  // Null max distance will safely disable smoothing transition calculations
+			}
+			return ClothingMeshUtils::ClothMeshDesc(SimPositions, SimIndices);  // Calculate averaged normals
+		};
+
+	// Get the min and max LODs that might be affected by LOD bias and update other sections' bias with this LOD section
+	int32 MinLOD = 0;
+	int32 MaxLOD = 0;
+
+	if (bAddMappingsToAnyLOD)
+	{
+		MinLOD = UpdatedLODIndex + 1;
+		MaxLOD = LODModels.Num() - 1;
+	}
+	else
+	{
+		const int32 RayTracingMinLOD = SkeletalMesh->GetRayTracingMinLOD();
+		if (bAddMappingsToMinLOD && LODModels.IsValidIndex(RayTracingMinLOD))
+		{
+			MinLOD = MaxLOD = RayTracingMinLOD;
+		}
+	}
+	if (UpdatedLODIndex < MinLOD && MinLOD <= MaxLOD)
+	{
+		// Prepare the deformer source (simulation) mesh descriptor
+		FPointWeightMap* MaxDistances;
+		const ClothingMeshUtils::ClothMeshDesc SourceMesh = PrepareDeformerSourceDescriptor(Asset, UpdatedLODIndex, MaxDistances);
+
+		// Generate the mappings for this section at the various LODs using the current simulation mesh
+		// so that higher LODs have something to map to when raytracing bias is active and a higher LOD is rendered
+		for (int32 LODIndex = MinLOD; LODIndex <= MaxLOD; ++LODIndex)
+		{
+			FSkeletalMeshLODModel& LODModel = LODModels[LODIndex];  // Note that the section LOD being updated here is the biased (raytraced) section LOD, not the UpdatedLODIndex which is done in the second loop below
+
+			if (!LODModel.Sections.IsValidIndex(SectionIndex))
+			{
+				continue;
+			}
+
+			FSkelMeshSection& Section = LODModel.Sections[SectionIndex];
+
+			if (!Section.HasClothingData())  // Note it doesn't matter if this section has a different asset at this point, only that the shader will want to access the clothing data
+			{
+				continue;  // This LOD section won't render clothing and therefore can be skipped as it won't ever need mappings
+			}
+
+			// Update bias mapping to the newly updated LOD index
+			const int32 LODBias = LODIndex - UpdatedLODIndex;
+			check(LODBias > 0);
+
+			Section.ClothMappingDataLODs.SetNum(FMath::Max(LODBias + 1, Section.ClothMappingDataLODs.Num()));
+
+			TArray<FMeshToMeshVertData>& MeshToMeshData = Section.ClothMappingDataLODs[LODBias];
+			MeshToMeshData.Reset();
+
+			// Prepare the deformer target (render) mesh descriptor
+			const ClothingMeshUtils::ClothMeshDesc TargetMesh = PrepareDeformerTargetDescriptor(LODIndex, SectionIndex);
+
+			// Generate the missing bias mapping data
+			ClothingMeshUtils::GenerateMeshToMeshVertData(
+				MeshToMeshData,
+				TargetMesh,
+				SourceMesh,
+				MaxDistances,
+				bSmoothTransition,
+				bUseMultipleInfluences,
+				SkinningKernelRadius);
+
+			UE_LOG(LogClothingAsset, Verbose, TEXT("Added deformer data for [%s/%s], section %d, LODBias = %d, render mesh LOD = %d, sim mesh LOD = %d"),
+				*SkeletalMesh->GetName(),
+				*Asset.GetName(),
+				SectionIndex,
+				LODBias,
+				LODIndex,
+				UpdatedLODIndex);
+		}
+	}
+
+	// Update this LOD section bias mappings that didn't get setup since there weren't no cloth attached to it yet
+	if (bAddMappingsToAnyLOD || (bAddMappingsToMinLOD && UpdatedLODIndex == SkeletalMesh->GetRayTracingMinLOD()))
+	{
+		FSkeletalMeshLODModel& LODModel = LODModels[UpdatedLODIndex];  // Note that the section LOD being updated here is the UpdatedLODIndex
+		FSkelMeshSection& Section = LODModel.Sections[SectionIndex];   // not the biased (raytraced) section LOD which is done in the first loop above
+		check(Section.HasClothingData());
+
+		// Prepare the deformer target (render) mesh descriptor
+		const ClothingMeshUtils::ClothMeshDesc TargetMesh = PrepareDeformerTargetDescriptor(UpdatedLODIndex, SectionIndex);
+
+		// Iterate through all the LOD that could be rendered and need this source LOD to be raytraced
+		for (int32 LODIndex = 0; LODIndex < UpdatedLODIndex; ++LODIndex)
+		{
+			FSkeletalMeshLODModel& SourceLODModel = LODModels[LODIndex];
+
+			if (!SourceLODModel.Sections.IsValidIndex(SectionIndex))
+			{
+				continue;
+			}
+
+			FSkelMeshSection& SourceSection = SourceLODModel.Sections[SectionIndex];
+
+			// Locate the deformer's source (simulation) cloth asset, since it could be a different cloth asset.
+			const UClothingAssetBase* const ClothingAsset = SkeletalMesh->GetClothingAsset(SourceSection.ClothingData.AssetGuid);
+			if (!ClothingAsset)
+			{
+				UE_LOG(LogClothingAsset, Warning, TEXT("Incomplete raytracing cloth LOD bias mappings generation on [%s/%s], section %d, LOD %d due to no clothing data set to LOD %d."),
+					*SkeletalMesh->GetName(),
+					*Asset.GetName(),
+					SectionIndex,
+					UpdatedLODIndex,
+					LODIndex);
+				continue;
+			}
+
+			// Prepare the deformer source (simulation) mesh descriptor
+			FPointWeightMap* MaxDistances;
+			const ClothingMeshUtils::ClothMeshDesc SourceMesh = PrepareDeformerSourceDescriptor(*ClothingAsset, LODIndex, MaxDistances);
+
+			// Update bias mapping for the updated LOD section
+			const int32 LODBias = UpdatedLODIndex - LODIndex;
+			check(LODBias > 0);
+
+			Section.ClothMappingDataLODs.SetNum(FMath::Max(LODBias + 1, Section.ClothMappingDataLODs.Num()));
+			TArray<FMeshToMeshVertData>& MeshToMeshData = Section.ClothMappingDataLODs[LODBias];
+			MeshToMeshData.Reset();
+
+			ClothingMeshUtils::GenerateMeshToMeshVertData(
+				MeshToMeshData,
+				TargetMesh,
+				SourceMesh,
+				MaxDistances,
+				bSmoothTransition,
+				bUseMultipleInfluences,
+				SkinningKernelRadius);
+
+			UE_LOG(LogClothingAsset, Verbose, TEXT("Added deformer data for [%s/%s], section %d, LODBias = %d, render mesh LOD = %d, sim mesh LOD = %d"),
+				*SkeletalMesh->GetName(),
+				*Asset.GetName(),
+				SectionIndex,
+				LODBias,
+				UpdatedLODIndex,
+				LODIndex);
+		}
+	}
+}
+
 // UE_DEPRECATED(5.7, "Unused low level function, will be removed or re-added as a member to FSkelMeshSection.")
 void ClothingAssetUtils::ClearSectionClothingData(FSkelMeshSection& Section)
 {
@@ -516,7 +735,14 @@
 	if (InSkelMesh->GetSupportRayTracing())
 	{
 		check(LodMap[InMeshLodIndex] == InAssetLodIndex);  // UpdateLODBiasMappings relies on the LodMap being up to date
-		UpdateLODBiasMappings(InSkelMesh, InMeshLodIndex, InSectionIndex);
+		ClothingAssetUtils::UpdateLODBiasMappings(
+			*this,
+			InSkelMesh,
+			InMeshLodIndex,
+			InSectionIndex,
+			ClothLodData.bSmoothTransition,
+			ClothLodData.bUseMultipleInfluences,
+			ClothLodData.SkinningKernelRadius);
 	}
 
 	return true;
@@ -535,202 +761,57 @@
 	}
 
 	// Iterate through all source LODs with cloth that could lead to upper (raytraced) sections needing some additional mapping
-	TIndirectArray<FSkeletalMeshLODModel>& LODModels = SkeletalMesh->GetImportedModel()->LODModels;
+	TIndirectArray<FSkeletalMeshLODModel>& LODModels = SkeletalMeshModel->LODModels;
 
 	for (int32 LODIndex = LODModels.Num() - 1; LODIndex >= 0; --LODIndex)  // Iterate in reverse order to allow shrinking the mapping array first
 	{
-		// Go through all sections to find the one(s?) that uses this cloth asset and clear the existing bias mappings
-		TArray<FSkelMeshSection>& Sections = LODModels[LODIndex].Sections;
+		if (LodMap.IsValidIndex(LODIndex) && LodData.IsValidIndex(LodMap[LODIndex]))
+		{
+			const FClothLODDataCommon& ClothLodData = LodData[LodMap[LODIndex]];
+
+			// Go through all sections to find the one(s?) that uses this cloth asset and clear the existing bias mappings
+			TArray<FSkelMeshSection>& Sections = LODModels[LODIndex].Sections;
 		
-		for (int32 SectionIndex = 0; SectionIndex < Sections.Num(); ++SectionIndex)
-		{
-			if (Sections[SectionIndex].ClothingData.AssetGuid == GetAssetGuid())
+			for (int32 SectionIndex = 0; SectionIndex < Sections.Num(); ++SectionIndex)
 			{
-				Sections[SectionIndex].ClothMappingDataLODs.SetNum(1);  // Only keep ClothMappingDataLODs[0] that is the same LOD mapping to remove LOD bias mappings
-
-				if (SkeletalMesh->GetSupportRayTracing())
+				if (Sections[SectionIndex].ClothingData.AssetGuid == GetAssetGuid())
 				{
-					UpdateLODBiasMappings(SkeletalMesh, LODIndex, SectionIndex);  // Updates the upper LODs mappings of the specified section from this LODIndex
+					Sections[SectionIndex].ClothMappingDataLODs.SetNum(1);  // Only keep ClothMappingDataLODs[0] that is the same LOD mapping to remove LOD bias mappings
+
+					if (SkeletalMesh->GetSupportRayTracing())  // This conditional is inside the loop so that the SetNum(1) can removed the unused bias mappings
+					{
+						// Updates the upper LODs mappings of the specified section from this LODIndex
+						ClothingAssetUtils::UpdateLODBiasMappings(
+							*this,
+							SkeletalMesh,
+							LODIndex,
+							SectionIndex,
+							ClothLodData.bSmoothTransition,
+							ClothLodData.bUseMultipleInfluences,
+							ClothLodData.SkinningKernelRadius);
+					}
 				}
 			}
 		}
 	}
 }
 
-void UClothingAssetCommon::UpdateLODBiasMappings(const USkeletalMesh* SkeletalMesh, int32 UpdatedLODIndex, int32 SectionIndex)
+void UClothingAssetCommon::GetSimulationMesh(const int32 InMeshLodIndex, TArray<FVector3f>& OutPositions, TArray<uint32>& OutIndices, TArray<float>& OutMaxDistances)const
 {
-	check(SkeletalMesh);
-	check(UpdatedLODIndex >= 0);
-	
-	FSkeletalMeshModel* const SkeletalMeshModel = SkeletalMesh->GetImportedModel();
-	check(SkeletalMeshModel);
+	OutPositions.Reset();
+	OutIndices.Reset();
+	OutMaxDistances.Reset();
 
-	const bool bAddMappingsToAnyLOD = (SkeletalMesh->GetClothLODBiasMode() == EClothLODBiasMode::MappingsToAnyLOD);
-	const bool bAddMappingsToMinLOD = (SkeletalMesh->GetClothLODBiasMode() == EClothLODBiasMode::MappingsToMinLOD && SkeletalMesh->GetRayTracingMinLOD() > 0);
+	if (LodMap.IsValidIndex(InMeshLodIndex) && LodData.IsValidIndex(LodMap[InMeshLodIndex]))
+	{
+		const FClothPhysicalMeshData& PhysicalMeshData = LodData[LodMap[InMeshLodIndex]].PhysicalMeshData;
 
-	TIndirectArray<FSkeletalMeshLODModel>& LODModels = SkeletalMeshModel->LODModels;
-
-	TArray<FVector3f> RenderPositions;
-	TArray<FVector3f> RenderNormals;
-	TArray<FVector3f> RenderTangents;
-	TArray<uint32> RenderIndices;
-
-	auto PrepareDeformerTargetDescriptor = [&LODModels, &RenderPositions, &RenderNormals, &RenderTangents, &RenderIndices](int32 LODIndex, int32 SectionIndex)
-	{
-		FSkeletalMeshLODModel& LODModel = LODModels[LODIndex];
-		FSkelMeshSection& Section = LODModel.Sections[SectionIndex];
-
-		RenderPositions.Reset(Section.SoftVertices.Num());
-		RenderNormals.Reset(Section.SoftVertices.Num());
-		RenderTangents.Reset(Section.SoftVertices.Num());
-
-		for (const FSoftSkinVertex& SoftSkinVertex : Section.SoftVertices)
+		OutPositions = PhysicalMeshData.Vertices;
+		OutIndices = PhysicalMeshData.Indices;
+		if (const FPointWeightMap* const MaxDistances = PhysicalMeshData.FindWeightMap(EWeightMapTargetCommon::MaxDistance);
+			MaxDistances && MaxDistances->Num() == OutPositions.Num())
 		{
-			RenderPositions.Add(SoftSkinVertex.Position);
-			RenderNormals.Add(SoftSkinVertex.TangentZ);
-			RenderTangents.Add(SoftSkinVertex.TangentX);
-		}
-
-		const TConstArrayView<uint32> Indices = TConstArrayView<uint32>(LODModel.IndexBuffer).Slice(Section.BaseIndex, Section.NumTriangles * 3);
-
-		RenderIndices.Reset(Section.NumTriangles * 3);
-		for (uint32 VertexIndex : Indices)
-		{
-			RenderIndices.Add(VertexIndex - Section.BaseVertexIndex);
-		}
-
-		return ClothingMeshUtils::ClothMeshDesc(RenderPositions, RenderNormals, RenderTangents, RenderIndices);
-	};
-
-	// Get the min and max LODs that might be affected by LOD bias and update other sections' bias with this LOD section
-	int32 MinLOD = 0;
-	int32 MaxLOD = 0;
-	
-	if (bAddMappingsToAnyLOD)
-	{
-		MinLOD = UpdatedLODIndex + 1;
-		MaxLOD = LODModels.Num() - 1;
-	}
-	else if (bAddMappingsToMinLOD)
-	{
-		MinLOD = MaxLOD = SkeletalMesh->GetRayTracingMinLOD();
-	}
-
-	if (UpdatedLODIndex < MinLOD && MinLOD <= MaxLOD)
-	{
-		// Prepare the deformer source (simulation) mesh descriptor
-		const FClothLODDataCommon& ClothLodData = LodData[LodMap[UpdatedLODIndex]];
-		const ClothingMeshUtils::ClothMeshDesc SourceMesh(ClothLodData.PhysicalMeshData.Vertices, ClothLodData.PhysicalMeshData.Indices);  // Calculate averaged normals
-
-		// Retrieve max distance mask
-		const FPointWeightMap* const MaxDistances = ClothLodData.PhysicalMeshData.FindWeightMap(EWeightMapTargetCommon::MaxDistance);
-
-		// Iterate through all deformer target (render) mesh sections
-		for (int32 LODIndex = MinLOD; LODIndex <= MaxLOD; ++LODIndex)
-		{
-			FSkeletalMeshLODModel& LODModel = LODModels[LODIndex];  // Note that the section LOD being updated here is the biased (raytraced) section LOD, not the UpdatedLODIndex which is done in the second loop below
-
-			if (!LODModel.Sections.IsValidIndex(SectionIndex))
-			{
-				continue;
-			}
-
-			FSkelMeshSection& Section = LODModel.Sections[SectionIndex];
-
-			if (!Section.HasClothingData())
-			{
-				continue;  // This LOD section won't render clothing and therefore can be skipped as it won't ever need mappings
-			}
-
-			// Update bias mapping to the newly updated LOD index
-			const int32 LODBias = LODIndex - UpdatedLODIndex;
-			check(LODBias > 0);
-
-			Section.ClothMappingDataLODs.SetNum(FMath::Max(LODBias + 1, Section.ClothMappingDataLODs.Num()));
-
-			TArray<FMeshToMeshVertData>& MeshToMeshData = Section.ClothMappingDataLODs[LODBias];
-			MeshToMeshData.Reset();
-
-			// Prepare the deformer target (render) mesh descriptor
-			const ClothingMeshUtils::ClothMeshDesc TargetMesh = PrepareDeformerTargetDescriptor(LODIndex, SectionIndex);
-
-			// Generate the missing bias mapping data
-			ClothingMeshUtils::GenerateMeshToMeshVertData(MeshToMeshData, TargetMesh, SourceMesh, MaxDistances,
-				ClothLodData.bSmoothTransition, ClothLodData.bUseMultipleInfluences, ClothLodData.SkinningKernelRadius);
-
-			UE_LOG(LogClothingAsset, Verbose, TEXT("Added deformer data for [%s/%s], section %d, LODBias = %d, render mesh LOD = %d, sim mesh LOD = %d"),
-				*SkeletalMesh->GetName(),
-				*GetName(),
-				SectionIndex,
-				LODBias,
-				LODIndex,
-				UpdatedLODIndex);
-		}
-	}
-
-	// Update this LOD section bias mappings that didn't get setup since there weren't no cloth attached to it yet
-	if (bAddMappingsToAnyLOD || (bAddMappingsToMinLOD && UpdatedLODIndex == SkeletalMesh->GetRayTracingMinLOD()))
-	{
-		FSkeletalMeshLODModel& LODModel = LODModels[UpdatedLODIndex];  // Note that the section LOD being updated here is the UpdatedLODIndex, not the biased (raytraced) section LOD  which is done in the first loop above
-		FSkelMeshSection& Section = LODModel.Sections[SectionIndex];
-		check(Section.HasClothingData());
-
-		// Prepare the deformer target (render) mesh descriptor
-		const ClothingMeshUtils::ClothMeshDesc TargetMesh = PrepareDeformerTargetDescriptor(UpdatedLODIndex, SectionIndex);
-
-		// Iterate through all the LOD that could be rendered and need this source LOD to be raytraced
-		for (int32 LODIndex = 0; LODIndex < UpdatedLODIndex; ++LODIndex)
-		{
-			FSkeletalMeshLODModel& SourceLODModel = LODModels[LODIndex];
-
-			if (!SourceLODModel.Sections.IsValidIndex(SectionIndex))
-			{
-				continue;
-			}
-
-			FSkelMeshSection& SourceSection = SourceLODModel.Sections[SectionIndex];
-
-			// Locate the deformer's source (simulation) cloth asset, since it could be a different cloth asset.
-			const UClothingAssetCommon* const ClothingAsset = Cast<UClothingAssetCommon>(SkeletalMesh->GetClothingAsset(SourceSection.ClothingData.AssetGuid));
-
-			if (!ClothingAsset)
-			{
-				const FText Error = FText::Format(
-					LOCTEXT("Error_DifferentAsset", "Incomplete raytracing cloth LOD bias mappings generation on [{0}/{1}], section {2}, LOD {3} due to missing cloth asset at LOD {4}."),
-					FText::FromString(SkeletalMesh->GetName()),
-					FText::FromString(GetName()),
-					SectionIndex,
-					UpdatedLODIndex, 
-					LODIndex);
-				WarningNotification(Error);
-				continue;
-			}
-
-			// Prepare the deformer source (simulation) mesh descriptor
-			const FClothLODDataCommon& ClothLodData = ClothingAsset->LodData[LodMap[LODIndex]];
-			const ClothingMeshUtils::ClothMeshDesc SourceMesh(ClothLodData.PhysicalMeshData.Vertices, ClothLodData.PhysicalMeshData.Indices);  // Calculate averaged normals
-
-			// Retrieve max distance mask
-			const FPointWeightMap* const MaxDistances = ClothLodData.PhysicalMeshData.FindWeightMap(EWeightMapTargetCommon::MaxDistance);
-
-			// Update bias mapping for the updated LOD section
-			const int32 LODBias = UpdatedLODIndex - LODIndex;
-			check(LODBias > 0);
-
-			Section.ClothMappingDataLODs.SetNum(FMath::Max(LODBias + 1, Section.ClothMappingDataLODs.Num()));
-			TArray<FMeshToMeshVertData>& MeshToMeshData = Section.ClothMappingDataLODs[LODBias];
-			MeshToMeshData.Reset();
-
-			ClothingMeshUtils::GenerateMeshToMeshVertData(MeshToMeshData, TargetMesh, SourceMesh, MaxDistances,
-				ClothLodData.bSmoothTransition, ClothLodData.bUseMultipleInfluences, ClothLodData.SkinningKernelRadius);
-
-			UE_LOG(LogClothingAsset, Verbose, TEXT("Added deformer data for [%s/%s], section %d, LODBias = %d, render mesh LOD = %d, sim mesh LOD = %d"),
-				*SkeletalMesh->GetName(),
-				*GetName(),
-				SectionIndex,
-				LODBias,
-				UpdatedLODIndex,
-				LODIndex);
+			OutMaxDistances = MaxDistances->Values;
 		}
 	}
 }
diff --git a/Engine/Source/Runtime/ClothingSystemRuntimeCommon/Public/ClothingAsset.h b/Engine/Source/Runtime/ClothingSystemRuntimeCommon/Public/ClothingAsset.h
--- a/Engine/Source/Runtime/ClothingSystemRuntimeCommon/Public/ClothingAsset.h
+++ b/Engine/Source/Runtime/ClothingSystemRuntimeCommon/Public/ClothingAsset.h
@@ -76,13 +76,32 @@ namespace ClothingAssetUtils
 		int32 LodIndex);
 
 	/**
-	 * Clear all defomer section that relies on the specified cloth sim data bound to a LOD section.
+	 * Clear all deformer section that relies on the specified cloth sim data bound to a LOD section.
 	 * @param SkeletalMesh The skeletal mesh to clear.
 	 * @param UpdatedLodIndex The recently updated Skeletal Mesh LOD index from which the bias mapping needs to be cleared.
 	 * @param SectionIndex The section to clear.
 	 */
 	UE_API void ClearLODBiasMappings(USkeletalMesh& SkeletalMesh, const int32 UpdatedLodIndex, const int32 SectionIndex);
 
+	/**
+	 * Add extra cloth deformer mappings to cope with a different raytracing LOD than the one currently rendered.
+	 * @param Asset The Clothing Asset to use to generate the mapping of the section for all LODs above the updated LOD up to the max LOD.
+	 * @param SkeletalMesh The Skeletal Mesh for which to calculate the render deformer mapping
+	 * @param UpdatedLODIndex The Skeletal Mesh LOD that has just been updated. Only LOD from above this one will receive new bias mapping if they exist.
+	 * @param SectionIndex The Skeletal Mesh section index to update.
+	 * @param bSmoothTransition Whether to create a smooth transition area between skinned render vertices and deformed render vertices using the vertex mapping blend weight.
+	 * @param bUseMultipleInfluences Whether to use multiple simulation mesh triangles to set the deformer influences of the render vertices. 
+	 * @param SkinningKernelRadius The radius from which to search for the closest simulation mesh triangles.
+	 */
+	UE_API void UpdateLODBiasMappings(
+		const UClothingAssetBase& Asset,
+		const USkeletalMesh* SkeletalMesh,
+		const int32 UpdatedLODIndex,
+		const int32 SectionIndex, 
+		const bool bSmoothTransition = true,
+		const bool bUseMultipleInfluences = false,
+		const float SkinningKernelRadius = 20.f);
+
 	UE_DEPRECATED(5.7, "Unused low level function, will be removed or re-added as a member to FSkelMeshSection.")
 	UE_API void ClearSectionClothingData(FSkelMeshSection& Section);
 #endif
@@ -164,6 +183,15 @@ public:
 	 */
 	UE_API virtual void UpdateAllLODBiasMappings(USkeletalMesh* SkeletalMesh) override;
 
+	/**
+	 * Return the simulation mesh.
+	 * Used to calculate the LOD bias deformer mappings.
+	 * @param OutPositions The simulation mesh positions.
+	 * @param OutIndices The simulation mesh triangle indices.
+	 * @param OutMaxDistances The simulation mesh max distance attributes (0 for kinematic, >0 for dynamic). must match the number of positions.
+	 */
+	UE_API virtual void GetSimulationMesh(const int32 InMeshLodIndex, TArray<FVector3f>& OutPositions, TArray<uint32>& OutIndices, TArray<float>& OutMaxDistances) const override;
+
 	/** 
 	 * Callback envoked after weights have been edited.
 	 * Calls \c PushWeightsToMesh() on each \c ClothLodData, and invalidates cached data. 
@@ -357,9 +385,6 @@ private:
 	// Return all LODs' stats.
 	virtual TArray<TPair<FText, FText>> GetStats() const override;
 
-	// Add extra cloth deformer mappings to cope with a different raytracing LOD than the one currently rendered.
-	void UpdateLODBiasMappings(const USkeletalMesh* SkeletalMesh, int32 UpdatedLODIndex, int32 SectionIndex);
-
 	// Helper functions used in PostPropertyChangeCb
 	void ReregisterComponentsUsingClothing();
 	void ForEachInteractorUsingClothing(TFunction<void (UClothingSimulationInteractor*)> Func);
diff --git a/Engine/Source/Runtime/ClothingSystemRuntimeInterface/Public/ClothingAssetBase.h b/Engine/Source/Runtime/ClothingSystemRuntimeInterface/Public/ClothingAssetBase.h
--- a/Engine/Source/Runtime/ClothingSystemRuntimeInterface/Public/ClothingAssetBase.h
+++ b/Engine/Source/Runtime/ClothingSystemRuntimeInterface/Public/ClothingAssetBase.h
@@ -66,6 +66,17 @@ PRAGMA_ENABLE_DEPRECATION_WARNINGS
 	virtual void UpdateAllLODBiasMappings(USkeletalMesh* SkeletalMesh)
 	PURE_VIRTUAL(UClothingAssetBase::UpdateAllLODBiasMappings, return;);
 
+	/**
+	 * Return the simulation mesh.
+	 * Used to calculate the LOD bias deformer mappings.
+	 * @param InMeshLodIndex The bound skeletal mesh LOD (in case this clothing data hold multiple LODs).
+	 * @param OutPositions The simulation mesh positions.
+	 * @param OutIndices The simulation mesh triangle indices.
+	 * @param OutMaxDistances The simulation mesh max distance attributes (0 for kinematic, >0 for dynamic). must match the number of positions.
+	 */
+	virtual void GetSimulationMesh(const int32 InMeshLodIndex, TArray<FVector3f>& OutPositions, TArray<uint32>& OutIndices, TArray<float>& OutMaxDistances) const
+	PURE_VIRTUAL(UClothingAssetBase::GetSimulationMesh, return;);
+
 	/** Add a new LOD class instance. */
 	UE_DEPRECATED(5.7, "Use UClothingAssetCommon::AddNewLod() instead.")
 	virtual int32 AddNewLod()
